Added include directives to fix non-unity build

diff --git a/Engine/Source/Programs/UnrealFrontend/Private/Windows/WindowsUnrealFrontendMain.cpp b/Engine/Source/Programs/UnrealFrontend/Private/Windows/WindowsUnrealFrontendMain.cpp
--- a/Engine/Source/Programs/UnrealFrontend/Private/Windows/WindowsUnrealFrontendMain.cpp
+++ b/Engine/Source/Programs/UnrealFrontend/Private/Windows/WindowsUnrealFrontendMain.cpp
@@ -5,7 +5,8 @@
 =============================================================================*/
 
 #include "UnrealFrontendMain.h"
-
+#include "ExceptionHandling.h"
+#include "LaunchEngineLoop.h"
 
 /**
  * The main application entry point for Windows platforms.
