#ue4
#ufe
* added some modules to the UFE dependencies so they are properly built even tho there is no direct dependency

diff --git a/Engine/Source/Programs/UnrealFrontend/UnrealFrontend.Build.cs b/Engine/Source/Programs/UnrealFrontend/UnrealFrontend.Build.cs
--- a/Engine/Source/Programs/UnrealFrontend/UnrealFrontend.Build.cs
+++ b/Engine/Source/Programs/UnrealFrontend/UnrealFrontend.Build.cs
@@ -20,12 +20,16 @@ public class UnrealFrontend : ModuleRules
 				"AutomationController",
 				"Core",
 				"DeviceManager",
+				"LauncherServices",
+				"Messaging",
 				"ProfilerClient",
 				"Projects",
 				"SessionFrontend",
 				"SessionLauncher",
+				"SessionServices",
 				"Slate",
 				"StandaloneRenderer",
+				"TargetDeviceServices",
 				"TargetPlatform",
 			}
 		);
