fix build

diff --git a/Engine/Source/Programs/UnrealFrontend/Private/UnrealFrontendMain.cpp b/Engine/Source/Programs/UnrealFrontend/Private/UnrealFrontendMain.cpp
--- a/Engine/Source/Programs/UnrealFrontend/Private/UnrealFrontendMain.cpp
+++ b/Engine/Source/Programs/UnrealFrontend/Private/UnrealFrontendMain.cpp
@@ -269,7 +269,7 @@ void RunStatsConvertCommand()
 	}
 
 	// This is not supported yet.
-	if( Stream.Header.bRawStatFile )
+	if (Stream.Header.bRawStatsFile)
 	{
 		UE_LOG( LogStats, Error, TEXT( "Could not open input file, not supported type (raw): %s" ), *TargetFile );
 		return;
