CIS Fixes for UnrealFrontend linker errors + initialization order fix

diff --git a/Engine/Source/Programs/UnrealFrontend/UnrealFrontend.Build.cs b/Engine/Source/Programs/UnrealFrontend/UnrealFrontend.Build.cs
--- a/Engine/Source/Programs/UnrealFrontend/UnrealFrontend.Build.cs
+++ b/Engine/Source/Programs/UnrealFrontend/UnrealFrontend.Build.cs
@@ -19,6 +19,7 @@ public class UnrealFrontend : ModuleRules
 			new string[] {
 				"AutomationController",
 				"Core",
+				"CoreUObject",
 				"DeviceManager",
 				"LauncherServices",
 				"Messaging",
