Removed Session Launcher

diff --git a/Engine/Source/Developer/SessionLauncher/Private/Models/SessionLauncherCommands.h b/Engine/Source/Developer/SessionLauncher/Private/Models/SessionLauncherCommands.h
deleted file mode 100644
diff --git a/Engine/Source/Developer/SessionLauncher/Private/Models/SessionLauncherModel.h b/Engine/Source/Developer/SessionLauncher/Private/Models/SessionLauncherModel.h
deleted file mode 100644
diff --git a/Engine/Source/Developer/SessionLauncher/Private/SessionLauncherModule.cpp b/Engine/Source/Developer/SessionLauncher/Private/SessionLauncherModule.cpp
deleted file mode 100644
diff --git a/Engine/Source/Developer/SessionLauncher/Private/SessionLauncherPrivatePCH.h b/Engine/Source/Developer/SessionLauncher/Private/SessionLauncherPrivatePCH.h
deleted file mode 100644
diff --git a/Engine/Source/Developer/SessionLauncher/Private/Widgets/Build/SSessionLauncherBuildPage.cpp b/Engine/Source/Developer/SessionLauncher/Private/Widgets/Build/SSessionLauncherBuildPage.cpp
deleted file mode 100644
diff --git a/Engine/Source/Developer/SessionLauncher/Private/Widgets/Build/SSessionLauncherBuildPage.h b/Engine/Source/Developer/SessionLauncher/Private/Widgets/Build/SSessionLauncherBuildPage.h
deleted file mode 100644
diff --git a/Engine/Source/Developer/SessionLauncher/Private/Widgets/Cook/SSessionLauncherCookByTheBookSettings.cpp b/Engine/Source/Developer/SessionLauncher/Private/Widgets/Cook/SSessionLauncherCookByTheBookSettings.cpp
deleted file mode 100644
diff --git a/Engine/Source/Developer/SessionLauncher/Private/Widgets/Cook/SSessionLauncherCookByTheBookSettings.h b/Engine/Source/Developer/SessionLauncher/Private/Widgets/Cook/SSessionLauncherCookByTheBookSettings.h
deleted file mode 100644
diff --git a/Engine/Source/Developer/SessionLauncher/Private/Widgets/Cook/SSessionLauncherCookOnTheFlySettings.cpp b/Engine/Source/Developer/SessionLauncher/Private/Widgets/Cook/SSessionLauncherCookOnTheFlySettings.cpp
deleted file mode 100644
diff --git a/Engine/Source/Developer/SessionLauncher/Private/Widgets/Cook/SSessionLauncherCookOnTheFlySettings.h b/Engine/Source/Developer/SessionLauncher/Private/Widgets/Cook/SSessionLauncherCookOnTheFlySettings.h
deleted file mode 100644
diff --git a/Engine/Source/Developer/SessionLauncher/Private/Widgets/Cook/SSessionLauncherCookPage.cpp b/Engine/Source/Developer/SessionLauncher/Private/Widgets/Cook/SSessionLauncherCookPage.cpp
deleted file mode 100644
diff --git a/Engine/Source/Developer/SessionLauncher/Private/Widgets/Cook/SSessionLauncherCookPage.h b/Engine/Source/Developer/SessionLauncher/Private/Widgets/Cook/SSessionLauncherCookPage.h
deleted file mode 100644
diff --git a/Engine/Source/Developer/SessionLauncher/Private/Widgets/Cook/SSessionLauncherCookedPlatforms.h b/Engine/Source/Developer/SessionLauncher/Private/Widgets/Cook/SSessionLauncherCookedPlatforms.h
deleted file mode 100644
diff --git a/Engine/Source/Developer/SessionLauncher/Private/Widgets/Cook/SSessionLauncherCultureListRow.h b/Engine/Source/Developer/SessionLauncher/Private/Widgets/Cook/SSessionLauncherCultureListRow.h
deleted file mode 100644
diff --git a/Engine/Source/Developer/SessionLauncher/Private/Widgets/Cook/SSessionLauncherMapListRow.h b/Engine/Source/Developer/SessionLauncher/Private/Widgets/Cook/SSessionLauncherMapListRow.h
deleted file mode 100644
diff --git a/Engine/Source/Developer/SessionLauncher/Private/Widgets/Cook/SSessionLauncherPlatformListRow.h b/Engine/Source/Developer/SessionLauncher/Private/Widgets/Cook/SSessionLauncherPlatformListRow.h
deleted file mode 100644
diff --git a/Engine/Source/Developer/SessionLauncher/Private/Widgets/Cook/SSessionLauncherSimpleCookPage.cpp b/Engine/Source/Developer/SessionLauncher/Private/Widgets/Cook/SSessionLauncherSimpleCookPage.cpp
deleted file mode 100644
diff --git a/Engine/Source/Developer/SessionLauncher/Private/Widgets/Cook/SSessionLauncherSimpleCookPage.h b/Engine/Source/Developer/SessionLauncher/Private/Widgets/Cook/SSessionLauncherSimpleCookPage.h
deleted file mode 100644
diff --git a/Engine/Source/Developer/SessionLauncher/Private/Widgets/Deploy/SSessionLauncherDeployFileServerSettings.cpp b/Engine/Source/Developer/SessionLauncher/Private/Widgets/Deploy/SSessionLauncherDeployFileServerSettings.cpp
deleted file mode 100644
diff --git a/Engine/Source/Developer/SessionLauncher/Private/Widgets/Deploy/SSessionLauncherDeployFileServerSettings.h b/Engine/Source/Developer/SessionLauncher/Private/Widgets/Deploy/SSessionLauncherDeployFileServerSettings.h
deleted file mode 100644
diff --git a/Engine/Source/Developer/SessionLauncher/Private/Widgets/Deploy/SSessionLauncherDeployPage.cpp b/Engine/Source/Developer/SessionLauncher/Private/Widgets/Deploy/SSessionLauncherDeployPage.cpp
deleted file mode 100644
diff --git a/Engine/Source/Developer/SessionLauncher/Private/Widgets/Deploy/SSessionLauncherDeployPage.h b/Engine/Source/Developer/SessionLauncher/Private/Widgets/Deploy/SSessionLauncherDeployPage.h
deleted file mode 100644
diff --git a/Engine/Source/Developer/SessionLauncher/Private/Widgets/Deploy/SSessionLauncherDeployRepositorySettings.cpp b/Engine/Source/Developer/SessionLauncher/Private/Widgets/Deploy/SSessionLauncherDeployRepositorySettings.cpp
deleted file mode 100644
diff --git a/Engine/Source/Developer/SessionLauncher/Private/Widgets/Deploy/SSessionLauncherDeployRepositorySettings.h b/Engine/Source/Developer/SessionLauncher/Private/Widgets/Deploy/SSessionLauncherDeployRepositorySettings.h
deleted file mode 100644
diff --git a/Engine/Source/Developer/SessionLauncher/Private/Widgets/Deploy/SSessionLauncherDeployTargets.cpp b/Engine/Source/Developer/SessionLauncher/Private/Widgets/Deploy/SSessionLauncherDeployTargets.cpp
deleted file mode 100644
diff --git a/Engine/Source/Developer/SessionLauncher/Private/Widgets/Deploy/SSessionLauncherDeployTargets.h b/Engine/Source/Developer/SessionLauncher/Private/Widgets/Deploy/SSessionLauncherDeployTargets.h
deleted file mode 100644
diff --git a/Engine/Source/Developer/SessionLauncher/Private/Widgets/Deploy/SSessionLauncherDeployToDeviceSettings.cpp b/Engine/Source/Developer/SessionLauncher/Private/Widgets/Deploy/SSessionLauncherDeployToDeviceSettings.cpp
deleted file mode 100644
diff --git a/Engine/Source/Developer/SessionLauncher/Private/Widgets/Deploy/SSessionLauncherDeployToDeviceSettings.h b/Engine/Source/Developer/SessionLauncher/Private/Widgets/Deploy/SSessionLauncherDeployToDeviceSettings.h
deleted file mode 100644
diff --git a/Engine/Source/Developer/SessionLauncher/Private/Widgets/Deploy/SSessionLauncherDeviceGroupSelector.h b/Engine/Source/Developer/SessionLauncher/Private/Widgets/Deploy/SSessionLauncherDeviceGroupSelector.h
deleted file mode 100644
diff --git a/Engine/Source/Developer/SessionLauncher/Private/Widgets/Deploy/SSessionLauncherDeviceListRow.h b/Engine/Source/Developer/SessionLauncher/Private/Widgets/Deploy/SSessionLauncherDeviceListRow.h
deleted file mode 100644
diff --git a/Engine/Source/Developer/SessionLauncher/Private/Widgets/Launch/SSessionLauncherLaunchCustomRoles.cpp b/Engine/Source/Developer/SessionLauncher/Private/Widgets/Launch/SSessionLauncherLaunchCustomRoles.cpp
deleted file mode 100644
diff --git a/Engine/Source/Developer/SessionLauncher/Private/Widgets/Launch/SSessionLauncherLaunchCustomRoles.h b/Engine/Source/Developer/SessionLauncher/Private/Widgets/Launch/SSessionLauncherLaunchCustomRoles.h
deleted file mode 100644
diff --git a/Engine/Source/Developer/SessionLauncher/Private/Widgets/Launch/SSessionLauncherLaunchPage.cpp b/Engine/Source/Developer/SessionLauncher/Private/Widgets/Launch/SSessionLauncherLaunchPage.cpp
deleted file mode 100644
diff --git a/Engine/Source/Developer/SessionLauncher/Private/Widgets/Launch/SSessionLauncherLaunchPage.h b/Engine/Source/Developer/SessionLauncher/Private/Widgets/Launch/SSessionLauncherLaunchPage.h
deleted file mode 100644
diff --git a/Engine/Source/Developer/SessionLauncher/Private/Widgets/Launch/SSessionLauncherLaunchRoleEditor.h b/Engine/Source/Developer/SessionLauncher/Private/Widgets/Launch/SSessionLauncherLaunchRoleEditor.h
deleted file mode 100644
diff --git a/Engine/Source/Developer/SessionLauncher/Private/Widgets/Package/SSessionLauncherPackagePage.cpp b/Engine/Source/Developer/SessionLauncher/Private/Widgets/Package/SSessionLauncherPackagePage.cpp
deleted file mode 100644
diff --git a/Engine/Source/Developer/SessionLauncher/Private/Widgets/Package/SSessionLauncherPackagePage.h b/Engine/Source/Developer/SessionLauncher/Private/Widgets/Package/SSessionLauncherPackagePage.h
deleted file mode 100644
diff --git a/Engine/Source/Developer/SessionLauncher/Private/Widgets/Package/SSessionLauncherPackagingSettings.cpp b/Engine/Source/Developer/SessionLauncher/Private/Widgets/Package/SSessionLauncherPackagingSettings.cpp
deleted file mode 100644
diff --git a/Engine/Source/Developer/SessionLauncher/Private/Widgets/Package/SSessionLauncherPackagingSettings.h b/Engine/Source/Developer/SessionLauncher/Private/Widgets/Package/SSessionLauncherPackagingSettings.h
deleted file mode 100644
diff --git a/Engine/Source/Developer/SessionLauncher/Private/Widgets/Project/SSessionLauncherProjectPage.cpp b/Engine/Source/Developer/SessionLauncher/Private/Widgets/Project/SSessionLauncherProjectPage.cpp
deleted file mode 100644
diff --git a/Engine/Source/Developer/SessionLauncher/Private/Widgets/Project/SSessionLauncherProjectPage.h b/Engine/Source/Developer/SessionLauncher/Private/Widgets/Project/SSessionLauncherProjectPage.h
deleted file mode 100644
diff --git a/Engine/Source/Developer/SessionLauncher/Private/Widgets/Project/SSessionLauncherProjectPicker.cpp b/Engine/Source/Developer/SessionLauncher/Private/Widgets/Project/SSessionLauncherProjectPicker.cpp
deleted file mode 100644
diff --git a/Engine/Source/Developer/SessionLauncher/Private/Widgets/Project/SSessionLauncherProjectPicker.h b/Engine/Source/Developer/SessionLauncher/Private/Widgets/Project/SSessionLauncherProjectPicker.h
deleted file mode 100644
diff --git a/Engine/Source/Developer/SessionLauncher/Private/Widgets/SSessionLauncher.cpp b/Engine/Source/Developer/SessionLauncher/Private/Widgets/SSessionLauncher.cpp
deleted file mode 100644
diff --git a/Engine/Source/Developer/SessionLauncher/Private/Widgets/SSessionLauncher.h b/Engine/Source/Developer/SessionLauncher/Private/Widgets/SSessionLauncher.h
deleted file mode 100644
diff --git a/Engine/Source/Developer/SessionLauncher/Private/Widgets/SSessionLauncherBuildTaskSettings.cpp b/Engine/Source/Developer/SessionLauncher/Private/Widgets/SSessionLauncherBuildTaskSettings.cpp
deleted file mode 100644
diff --git a/Engine/Source/Developer/SessionLauncher/Private/Widgets/SSessionLauncherBuildTaskSettings.h b/Engine/Source/Developer/SessionLauncher/Private/Widgets/SSessionLauncherBuildTaskSettings.h
deleted file mode 100644
diff --git a/Engine/Source/Developer/SessionLauncher/Private/Widgets/SSessionLauncherDeployTaskSettings.cpp b/Engine/Source/Developer/SessionLauncher/Private/Widgets/SSessionLauncherDeployTaskSettings.cpp
deleted file mode 100644
diff --git a/Engine/Source/Developer/SessionLauncher/Private/Widgets/SSessionLauncherDeployTaskSettings.h b/Engine/Source/Developer/SessionLauncher/Private/Widgets/SSessionLauncherDeployTaskSettings.h
deleted file mode 100644
diff --git a/Engine/Source/Developer/SessionLauncher/Private/Widgets/SSessionLauncherLaunchTaskSettings.cpp b/Engine/Source/Developer/SessionLauncher/Private/Widgets/SSessionLauncherLaunchTaskSettings.cpp
deleted file mode 100644
diff --git a/Engine/Source/Developer/SessionLauncher/Private/Widgets/SSessionLauncherLaunchTaskSettings.h b/Engine/Source/Developer/SessionLauncher/Private/Widgets/SSessionLauncherLaunchTaskSettings.h
deleted file mode 100644
diff --git a/Engine/Source/Developer/SessionLauncher/Private/Widgets/SSessionLauncherMessageListRow.h b/Engine/Source/Developer/SessionLauncher/Private/Widgets/SSessionLauncherMessageListRow.h
deleted file mode 100644
diff --git a/Engine/Source/Developer/SessionLauncher/Private/Widgets/SSessionLauncherPreviewPage.cpp b/Engine/Source/Developer/SessionLauncher/Private/Widgets/SSessionLauncherPreviewPage.cpp
deleted file mode 100644
diff --git a/Engine/Source/Developer/SessionLauncher/Private/Widgets/SSessionLauncherPreviewPage.h b/Engine/Source/Developer/SessionLauncher/Private/Widgets/SSessionLauncherPreviewPage.h
deleted file mode 100644
diff --git a/Engine/Source/Developer/SessionLauncher/Private/Widgets/SSessionLauncherProfileSelector.h b/Engine/Source/Developer/SessionLauncher/Private/Widgets/SSessionLauncherProfileSelector.h
deleted file mode 100644
diff --git a/Engine/Source/Developer/SessionLauncher/Private/Widgets/SSessionLauncherProgress.h b/Engine/Source/Developer/SessionLauncher/Private/Widgets/SSessionLauncherProgress.h
deleted file mode 100644
diff --git a/Engine/Source/Developer/SessionLauncher/Private/Widgets/SSessionLauncherSettings.cpp b/Engine/Source/Developer/SessionLauncher/Private/Widgets/SSessionLauncherSettings.cpp
deleted file mode 100644
diff --git a/Engine/Source/Developer/SessionLauncher/Private/Widgets/SSessionLauncherSettings.h b/Engine/Source/Developer/SessionLauncher/Private/Widgets/SSessionLauncherSettings.h
deleted file mode 100644
diff --git a/Engine/Source/Developer/SessionLauncher/Private/Widgets/SSessionLauncherTaskListRow.h b/Engine/Source/Developer/SessionLauncher/Private/Widgets/SSessionLauncherTaskListRow.h
deleted file mode 100644
diff --git a/Engine/Source/Developer/SessionLauncher/Private/Widgets/SSessionLauncherValidation.h b/Engine/Source/Developer/SessionLauncher/Private/Widgets/SSessionLauncherValidation.h
deleted file mode 100644
diff --git a/Engine/Source/Developer/SessionLauncher/Private/Widgets/Shared/SSessionLauncherBuildConfigurationSelector.h b/Engine/Source/Developer/SessionLauncher/Private/Widgets/Shared/SSessionLauncherBuildConfigurationSelector.h
deleted file mode 100644
diff --git a/Engine/Source/Developer/SessionLauncher/Private/Widgets/Shared/SSessionLauncherFormLabel.h b/Engine/Source/Developer/SessionLauncher/Private/Widgets/Shared/SSessionLauncherFormLabel.h
deleted file mode 100644
diff --git a/Engine/Source/Developer/SessionLauncher/Private/Widgets/Toolbar/SSessionLauncherToolbar.cpp b/Engine/Source/Developer/SessionLauncher/Private/Widgets/Toolbar/SSessionLauncherToolbar.cpp
deleted file mode 100644
diff --git a/Engine/Source/Developer/SessionLauncher/Private/Widgets/Toolbar/SSessionLauncherToolbar.h b/Engine/Source/Developer/SessionLauncher/Private/Widgets/Toolbar/SSessionLauncherToolbar.h
deleted file mode 100644
diff --git a/Engine/Source/Developer/SessionLauncher/Public/Interfaces/ISessionLauncherModule.h b/Engine/Source/Developer/SessionLauncher/Public/Interfaces/ISessionLauncherModule.h
deleted file mode 100644
diff --git a/Engine/Source/Developer/SessionLauncher/Public/SessionLauncher.h b/Engine/Source/Developer/SessionLauncher/Public/SessionLauncher.h
deleted file mode 100644
diff --git a/Engine/Source/Developer/SessionLauncher/SessionLauncher.Build.cs b/Engine/Source/Developer/SessionLauncher/SessionLauncher.Build.cs
deleted file mode 100644
diff --git a/Engine/Source/Editor/MainFrame/Private/Menus/MainMenu.cpp b/Engine/Source/Editor/MainFrame/Private/Menus/MainMenu.cpp
--- a/Engine/Source/Editor/MainFrame/Private/Menus/MainMenu.cpp
+++ b/Engine/Source/Editor/MainFrame/Private/Menus/MainMenu.cpp
@@ -182,7 +182,6 @@ void FMainMenu::FillWindowMenu( FMenuBuilder& MenuBuilder, const TSharedRef< FEx
 
 		if (GetDefault<UEditorExperimentalSettings>()->bGameLauncher)
 		{
-			//FGlobalTabmanager::Get()->PopulateTabSpawnerMenu(MenuBuilder, "SessionLauncher"); // Will be deleted soon
 			FGlobalTabmanager::Get()->PopulateTabSpawnerMenu(MenuBuilder, "ProjectLauncher");
 		}
 
diff --git a/Engine/Source/Editor/UnrealEd/Private/Editor.cpp b/Engine/Source/Editor/UnrealEd/Private/Editor.cpp
--- a/Engine/Source/Editor/UnrealEd/Private/Editor.cpp
+++ b/Engine/Source/Editor/UnrealEd/Private/Editor.cpp
@@ -549,7 +549,6 @@ void UEditorEngine::Init(IEngineLoop* InEngineLoop)
 		FModuleManager::Get().LoadModule(TEXT("ProfilerClient"));
 //		FModuleManager::Get().LoadModule(TEXT("Search"));
 		FModuleManager::Get().LoadModule(TEXT("SessionFrontend"));
-		FModuleManager::Get().LoadModule(TEXT("SessionLauncher"));
 		FModuleManager::Get().LoadModule(TEXT("ProjectLauncher"));
 		FModuleManager::Get().LoadModule(TEXT("SettingsEditor"));
 		FModuleManager::Get().LoadModule(TEXT("EditorSettingsViewer"));
diff --git a/Engine/Source/Editor/UnrealEd/UnrealEd.Build.cs b/Engine/Source/Editor/UnrealEd/UnrealEd.Build.cs
--- a/Engine/Source/Editor/UnrealEd/UnrealEd.Build.cs
+++ b/Engine/Source/Editor/UnrealEd/UnrealEd.Build.cs
@@ -149,7 +149,6 @@ public class UnrealEd : ModuleRules
 				"DeviceManager",
 				"SettingsEditor",
 				"SessionFrontend",
-				"SessionLauncher",
 				"TaskBrowser",
 				"Sequencer",
 				"SoundClassEditor",
diff --git a/Engine/Source/Programs/UnrealFrontend/Private/UnrealFrontendMain.cpp b/Engine/Source/Programs/UnrealFrontend/Private/UnrealFrontendMain.cpp
--- a/Engine/Source/Programs/UnrealFrontend/Private/UnrealFrontendMain.cpp
+++ b/Engine/Source/Programs/UnrealFrontend/Private/UnrealFrontendMain.cpp
@@ -363,7 +363,6 @@ void RunUI()
 	// initialize user interface
 	FModuleManager::Get().LoadModule("DeviceManager");
 	FModuleManager::Get().LoadModule("SessionFrontend");
-	FModuleManager::Get().LoadModule("SessionLauncher");
 	FModuleManager::Get().LoadModule("SettingsEditor");
 	FModuleManager::Get().LoadModule("ProjectLauncher");
 
@@ -386,7 +385,6 @@ void RunUI()
 					FTabManager::NewStack()
 						->AddTab(FName("DeviceManager"), ETabState::OpenedTab)
 						->AddTab(FName("MessagingDebugger"), ETabState::ClosedTab)
-						->AddTab(FName("SessionLauncher"), ETabState::OpenedTab)
 						->AddTab(FName("SessionFrontend"), ETabState::OpenedTab)
 						->AddTab(FName("ProjectLauncher"), ETabState::OpenedTab)
 				)
diff --git a/Engine/Source/Programs/UnrealFrontend/UnrealFrontend.Build.cs b/Engine/Source/Programs/UnrealFrontend/UnrealFrontend.Build.cs
--- a/Engine/Source/Programs/UnrealFrontend/UnrealFrontend.Build.cs
+++ b/Engine/Source/Programs/UnrealFrontend/UnrealFrontend.Build.cs
@@ -27,7 +27,6 @@ public class UnrealFrontend : ModuleRules
                 "ProjectLauncher",
 				"Projects",
 				"SessionFrontend",
-				"SessionLauncher",
 				"SessionServices",
 				"Slate",
 				"SlateCore",
