Stats - Fixed converting stats to csv

diff --git a/Engine/Source/Programs/UnrealFrontend/Private/UnrealFrontendMain.cpp b/Engine/Source/Programs/UnrealFrontend/Private/UnrealFrontendMain.cpp
--- a/Engine/Source/Programs/UnrealFrontend/Private/UnrealFrontendMain.cpp
+++ b/Engine/Source/Programs/UnrealFrontend/Private/UnrealFrontendMain.cpp
@@ -200,23 +200,37 @@ bool RunLaunchCommand(const FString& Params)
 }
 
 
-void WriteString(FArchive* File, const ANSICHAR* Format, ...)
+struct FStatsConvertCommand
 {
-	if( File != NULL )
+	static void Run()
 	{
-		check(File);
-		ANSICHAR Array[1024];
-		va_list ArgPtr;
-		va_start(ArgPtr,Format);
-		// Build the string
-		int32 Result = FCStringAnsi::GetVarArgs(Array,ARRAY_COUNT(Array),ARRAY_COUNT(Array)-1,Format,ArgPtr);
-		// Now write that to the file
-		File->Serialize((void*)Array,Result);
+		FStatsConvertCommand Instance;
+		Instance.InternalRun();
 	}
-}
 
+protected:
+	void InternalRun();
+	void WriteString( FArchive& Writer, const ANSICHAR* Format, ... );
+	void CollectAndWriteStatsValues( FArchive& Writer );
+	void ReadAndConvertStatMessages( FArchive& Reader, FArchive& Writer );
 
-void RunStatsConvertCommand()
+	FStatsReadStream Stream;
+	FStatsThreadState ThreadState;
+	TArray<FName> StatList;
+};
+
+void FStatsConvertCommand::WriteString( FArchive& Writer, const ANSICHAR* Format, ... )
+{
+	ANSICHAR Array[1024];
+	va_list ArgPtr;
+	va_start(ArgPtr,Format);
+	// Build the string
+	int32 Result = FCStringAnsi::GetVarArgs(Array,ARRAY_COUNT(Array),ARRAY_COUNT(Array)-1,Format,ArgPtr);
+	// Now write that to the file
+	Writer.Serialize((void*)Array,Result);
+}
+
+void FStatsConvertCommand::InternalRun()
 {
 #if	STATS
 	// get the target file
@@ -228,10 +242,16 @@ void RunStatsConvertCommand()
 	FParse::Value(FCommandLine::Get(), TEXT("-STATLIST="), StatListString);
 
 	// get the list of stats
-	TArray<FString> StatList;
-	if (StatListString.ParseIntoArray(&StatList, TEXT("+"), true) == 0)
+	TArray<FString> StatArrayString;
+	if (StatListString.ParseIntoArray(&StatArrayString, TEXT("+"), true) == 0)
+	{
+		StatArrayString.Add(TEXT("STAT_FrameTime"));
+	}
+
+	// convert to FNames for faster compare
+	for( const FString& It : StatArrayString )
 	{
-		StatList.Add(TEXT("STAT_FrameTime"));
+		StatList.Add(*It);
 	}
 
 	// open a csv file for write
@@ -258,7 +278,6 @@ void RunStatsConvertCommand()
 		return;
 	}
 
-	FStatsReadStream Stream;
 	if( !Stream.ReadHeader( *FileReader ) )
 	{
 		UE_LOG( LogStats, Error, TEXT( "Could not open input file, bad magic: %s" ), *TargetFile );
@@ -272,69 +291,163 @@ void RunStatsConvertCommand()
 		return;
 	}
 
+	const bool bIsFinalized = Stream.Header.IsFinalized();
+	if( bIsFinalized )
+	{
+		// Read metadata.
+		TArray<FStatMessage> MetadataMessages;
+		Stream.ReadFNamesAndMetadataMessages( *FileReader, MetadataMessages );
+		ThreadState.ProcessMetaDataOnly( MetadataMessages );
+
+		// Read frames offsets.
+		Stream.ReadFramesOffsets( *FileReader );
+		FileReader->Seek( Stream.FramesInfo[0].FrameFileOffset );
+	}
+
+	if( Stream.Header.HasCompressedData() )
+	{
+		UE_CLOG( !bIsFinalized, LogStats, Fatal, TEXT( "Compressed stats file has to be finalized" ) );
+	}
+
+	ReadAndConvertStatMessages( *FileReader, *FileWriter );
+
+#endif // STATS
+}
+
+void FStatsConvertCommand::ReadAndConvertStatMessages( FArchive& Reader, FArchive& Writer )
+{
 	// output the csv header
-	WriteString( FileWriter, "Frame,Name,Value\r\n" );
+	WriteString( Writer, "Frame,Name,Value\r\n" );
 
-	// read in the data
+	uint64 ReadMessages = 0;
 	TArray<FStatMessage> Messages;
-	FStatsThreadState ThreadState;
-	while(FileReader->Tell() < FileReader->TotalSize())
+
+	// Buffer used to store the compressed and decompressed data.
+	TArray<uint8> SrcData;
+	TArray<uint8> DestData;
+
+	const bool bHasCompressedData = Stream.Header.HasCompressedData();
+	const bool bIsFinalized = Stream.Header.IsFinalized();
+	float DataLoadingProgress = 0.0f;
+
+	if( bHasCompressedData )
 	{
-		// read the message
-		FStatMessage Message(Stream.ReadMessage(*FileReader));
-		if (Message.NameAndInfo.GetShortName() != TEXT("Unknown FName"))
+		while( Reader.Tell() < Reader.TotalSize() )
 		{
-			if (Message.NameAndInfo.GetField<EStatOperation>() == EStatOperation::AdvanceFrameEventGameThread)
+			// Read the compressed data.
+			FCompressedStatsData UncompressedData( SrcData, DestData );
+			Reader << UncompressedData;
+			if( UncompressedData.HasReachedEndOfCompressedData() )
 			{
-				new (Messages) FStatMessage(Message);
-				ThreadState.AddMessages(Messages);
-				Messages.Reset(Messages.Num());
+				return;
+			}
 
-				// get the frame time and the render time
-				if (ThreadState.CurrentGameFrame > 1)
+			FMemoryReader MemoryReader( DestData, true );
+
+			while( MemoryReader.Tell() < MemoryReader.TotalSize() )
+			{
+				// read the message
+				FStatMessage Message( Stream.ReadMessage( MemoryReader, bIsFinalized ) );
+				ReadMessages++;
+				if( ReadMessages % 32768 == 0 )
+				{
+					UE_LOG( LogStats, Log, TEXT( "StatsConvertCommand progress: %.1f%%" ), DataLoadingProgress );
+				}
+
+				if( Message.NameAndInfo.GetShortName() != TEXT( "Unknown FName" ) )
 				{
-					// get the thread stats
-					TArray<FStatMessage> Stats;
-					ThreadState.GetInclusiveAggregateStackStats(ThreadState.CurrentGameFrame, Stats);
-					for (int32 Index = 0; Index < Stats.Num(); ++Index)
+					if( Message.NameAndInfo.GetField<EStatOperation>() == EStatOperation::AdvanceFrameEventGameThread && ReadMessages > 2 )
 					{
-						FStatMessage const& Meta = Stats[Index];
-//						UE_LOG(LogTemp, Display, TEXT("Stat: %s"), *Meta.NameAndInfo.GetShortName().ToString());
+						new (Messages) FStatMessage( Message );
+						ThreadState.AddMessages( Messages );
+						Messages.Reset();
 
-						for (int32 Jndex = 0; Jndex < StatList.Num(); ++Jndex)
-						{
-							if (Meta.NameAndInfo.GetShortName().ToString() == StatList[Jndex])
-							{
-								float CycleTime = 0.0f;
-								if (Meta.NameAndInfo.GetFlag(EStatMetaFlags::IsPackedCCAndDuration))
-								{
-									CycleTime = FPlatformTime::ToMilliseconds( FromPackedCallCountDuration_Duration(Meta.GetValue_int64()) );
-								}
-								else if (Meta.NameAndInfo.GetFlag(EStatMetaFlags::IsCycle))
-								{
-									CycleTime = FPlatformTime::ToMilliseconds( Meta.GetValue_int64() );
-								}
-								else
-								{
-									CycleTime = Meta.GetValue_int64();
-								}
-
-								// write out to the csv file
-								WriteString(FileWriter, "%d,%S,%f\r\n", ThreadState.CurrentGameFrame, *StatList[Jndex], CycleTime);
-							}
-						}
+						CollectAndWriteStatsValues( Writer );
+						DataLoadingProgress = (double)Reader.Tell() / (double)Reader.TotalSize() * 100.0f;
 					}
+
+					new (Messages) FStatMessage( Message );
 				}
+				else
+				{
+					break;
+				}
+			}
+		}
+	}
+	else
+	{
+		while( Reader.Tell() < Reader.TotalSize() )
+		{
+			// read the message
+			FStatMessage Message( Stream.ReadMessage( Reader, bIsFinalized ) );
+			ReadMessages++;
+			if( ReadMessages % 32768 == 0 )
+			{
+				UE_LOG( LogStats, Log, TEXT( "StatsConvertCommand progress: %.1f%%" ), DataLoadingProgress );
 			}
 
-			new (Messages) FStatMessage(Message);
+			if( Message.NameAndInfo.GetShortName() != TEXT( "Unknown FName" ) )
+			{
+				if( Message.NameAndInfo.GetField<EStatOperation>() == EStatOperation::SpecialMessageMarker )
+				{
+					// Simply break the loop.
+					// The profiler supports more advanced handling of this message.
+					return;
+				}
+				else if( Message.NameAndInfo.GetField<EStatOperation>() == EStatOperation::AdvanceFrameEventGameThread && ReadMessages > 2 )
+				{
+					new (Messages) FStatMessage( Message );
+					ThreadState.AddMessages( Messages );
+					Messages.Reset();
+
+					CollectAndWriteStatsValues( Writer );
+					DataLoadingProgress = (double)Reader.Tell() / (double)Reader.TotalSize() * 100.0f;
+				}
+
+				new (Messages) FStatMessage( Message );
+			}
+			else
+			{
+				break;
+			}
 		}
-		else
+	}
+}
+
+void FStatsConvertCommand::CollectAndWriteStatsValues( FArchive& Writer )
+{
+	// get the thread stats
+	TArray<FStatMessage> Stats;
+	ThreadState.GetInclusiveAggregateStackStats(ThreadState.CurrentGameFrame, Stats);
+	for (int32 Index = 0; Index < Stats.Num(); ++Index)
+	{
+		FStatMessage const& Meta = Stats[Index];
+		//UE_LOG(LogTemp, Display, TEXT("Stat: %s"), *Meta.NameAndInfo.GetShortName().ToString());
+
+		for (int32 Jndex = 0; Jndex < StatList.Num(); ++Jndex)
 		{
-			break;
+			if (Meta.NameAndInfo.GetShortName() == StatList[Jndex])
+			{
+				double StatValue = 0.0f;
+				if (Meta.NameAndInfo.GetFlag(EStatMetaFlags::IsPackedCCAndDuration))
+				{
+					StatValue = FPlatformTime::ToMilliseconds( FromPackedCallCountDuration_Duration(Meta.GetValue_int64()) );
+				}
+				else if (Meta.NameAndInfo.GetFlag(EStatMetaFlags::IsCycle))
+				{
+					StatValue = FPlatformTime::ToMilliseconds( Meta.GetValue_int64() );
+				}
+				else
+				{
+					StatValue = Meta.GetValue_int64();
+				}
+
+				// write out to the csv file
+				WriteString( Writer, "%d,%S,%f\r\n", ThreadState.CurrentGameFrame, *StatList[Jndex].ToString(), StatValue );
+			}
 		}
 	}
-#endif // STATS
 }
 
 void RunUI()
@@ -490,7 +603,7 @@ int32 UnrealFrontendMain( const TCHAR* CommandLine )
 		}
 		else if (Command.Equals(TEXT("CONVERT"), ESearchCase::IgnoreCase))
 		{
-			RunStatsConvertCommand();
+			FStatsConvertCommand::Run();
 		}
 	}
 	else
