UFE: Broke out commands into separate files; removed unnecessary dependencies

diff --git a/Engine/Source/Programs/UnrealFrontend/Private/Commands/DeployCommand.cpp b/Engine/Source/Programs/UnrealFrontend/Private/Commands/DeployCommand.cpp
new file mode 100644
--- /dev/null
+++ b/Engine/Source/Programs/UnrealFrontend/Private/Commands/DeployCommand.cpp
@@ -0,0 +1,96 @@
+// Copyright 1998-2014 Epic Games, Inc. All Rights Reserved.
+
+#include "UnrealFrontendPrivatePCH.h"
+#include "DeployCommand.h"
+#include "Projects.h"
+#include "TargetPlatform.h"
+#include "TaskGraphInterfaces.h"
+
+
+bool FDeployCommand::Run( )
+{
+	bool bDeployed = false;
+
+	// get the target device
+	FString Device;
+	FParse::Value(FCommandLine::Get(), TEXT("-DEVICE="), Device);
+
+	// get the file manifest
+	FString Manifest;
+	FParse::Value(FCommandLine::Get(), TEXT("-MANIFEST="), Manifest);
+
+	FString SourceDir;
+	FParse::Value(FCommandLine::Get(), TEXT("-SOURCEDIR="), SourceDir);
+
+	ITargetPlatformManagerModule* TPM = GetTargetPlatformManager();
+
+	if (!TPM)
+	{
+		return false;
+	}
+
+	// Initialize the messaging subsystem so we can do device discovery.
+	FModuleManager::Get().LoadModuleChecked("Messaging");
+
+	// load plug-in modules
+	// @todo: allow for better plug-in support in standalone Slate apps
+	IPluginManager::Get().LoadModulesForEnabledPlugins(ELoadingPhase::PreDefault);
+
+	double DeltaTime = 0.0;
+	double LastTime = FPlatformTime::Seconds();
+
+	// We track the message sent time because we have to keep updating the loop until the message is *actually sent*. (ie all packets queued, sent, buffer flushed, etc.)
+	double MessageSentTime = 0.0;
+	bool bMessageSent = false;
+
+	while (!GIsRequestingExit && ((MessageSentTime > LastTime + 1.0) || (MessageSentTime <= 0.1)))
+	{
+		FTaskGraphInterface::Get().ProcessThreadUntilIdle(ENamedThreads::GameThread);
+		FTicker::GetCoreTicker().Tick(DeltaTime);
+		FPlatformProcess::Sleep(0);
+
+		DeltaTime = FPlatformTime::Seconds() - LastTime;
+		LastTime = FPlatformTime::Seconds();
+
+		if ( !bMessageSent )
+		{
+			const TArray<ITargetPlatform*>& Platforms = TPM->GetActiveTargetPlatforms();
+
+			FString Platform;
+			FString DeviceName;
+
+			Device.Split(TEXT("@"), &Platform, &DeviceName);
+
+			FTargetDeviceId DeviceId(Platform, DeviceName);
+			ITargetDevicePtr TargetDevice;
+
+			for (int32 Index = 0; Index < Platforms.Num(); ++Index)
+			{
+				TargetDevice = Platforms[Index]->GetDevice(DeviceId);
+
+				if (TargetDevice.IsValid())
+				{
+					FString OutId;
+
+					if (Platforms[Index]->PackageBuild(SourceDir))
+					{
+						if (TargetDevice->Deploy(SourceDir, OutId))
+						{
+							bDeployed = true;
+						}
+
+						MessageSentTime = LastTime;
+						bMessageSent = true;
+					}
+					else
+					{
+						MessageSentTime = LastTime;
+						bMessageSent = true;
+					}
+				}
+			}
+		}
+	}
+
+	return bDeployed;
+}
diff --git a/Engine/Source/Programs/UnrealFrontend/Private/Commands/DeployCommand.h b/Engine/Source/Programs/UnrealFrontend/Private/Commands/DeployCommand.h
new file mode 100644
--- /dev/null
+++ b/Engine/Source/Programs/UnrealFrontend/Private/Commands/DeployCommand.h
@@ -0,0 +1,16 @@
+// Copyright 1998-2014 Epic Games, Inc. All Rights Reserved.
+
+#pragma once
+
+
+class FDeployCommand
+{
+public:
+
+	/**
+	 * Executes the command.
+	 *
+	 * @return true on success, false otherwise.
+	 */
+	static bool Run( );
+};
diff --git a/Engine/Source/Programs/UnrealFrontend/Private/Commands/LaunchCommand.cpp b/Engine/Source/Programs/UnrealFrontend/Private/Commands/LaunchCommand.cpp
new file mode 100644
--- /dev/null
+++ b/Engine/Source/Programs/UnrealFrontend/Private/Commands/LaunchCommand.cpp
@@ -0,0 +1,89 @@
+// Copyright 1998-2014 Epic Games, Inc. All Rights Reserved.
+
+#include "UnrealFrontendPrivatePCH.h"
+#include "LaunchCommand.h"
+#include "Projects.h"
+#include "TargetPlatform.h"
+#include "TaskGraphInterfaces.h"
+
+
+bool FLaunchCommand::Run( const FString& Params )
+{
+	bool bLaunched = false;
+
+	// get the target device
+	FString Device;
+	FParse::Value(FCommandLine::Get(), TEXT("-DEVICE="), Device);
+
+	// get the executable to launch
+	FString Executable;
+	FParse::Value(FCommandLine::Get(), TEXT("-EXE="), Executable);
+
+	ITargetPlatformManagerModule* TPM = GetTargetPlatformManager();
+
+	if (!TPM)
+	{
+		return false;
+	}
+
+	// Initialize the messaging subsystem so we can do device discovery.
+	FModuleManager::Get().LoadModuleChecked("Messaging");
+
+	// load plug-in modules
+	// @todo: allow for better plug-in support in standalone Slate apps
+	IPluginManager::Get().LoadModulesForEnabledPlugins(ELoadingPhase::PreDefault);
+
+	double DeltaTime = 0.0;
+	double LastTime = FPlatformTime::Seconds();
+	static int32 MasterDisableChangeTagStartFrame = -1;
+
+	// We track the message sent time because we have to keep updating the loop until
+	// the message is *actually sent*. (ie all packets queued, sent, buffer flushed, etc.)
+	double MessageSentTime = 0.0;
+	bool bMessageSent = false;
+
+	while (!GIsRequestingExit && ((MessageSentTime > LastTime + 1.0) || (MessageSentTime <= 0.1)))
+	{
+		FTaskGraphInterface::Get().ProcessThreadUntilIdle(ENamedThreads::GameThread);
+		FTicker::GetCoreTicker().Tick(DeltaTime);
+		FPlatformProcess::Sleep(0);
+
+		DeltaTime = FPlatformTime::Seconds() - LastTime;
+		LastTime = FPlatformTime::Seconds();
+
+		if (!bMessageSent)
+		{
+			const TArray<ITargetPlatform*>& Platforms = TPM->GetActiveTargetPlatforms();
+
+			FString Platform;
+			FString DeviceName;
+			Device.Split(TEXT("@"), &Platform, &DeviceName);
+			FTargetDeviceId DeviceId(Platform, DeviceName);
+			ITargetDevicePtr TargetDevice;
+
+			for (int32 Index = 0; Index < Platforms.Num(); ++Index)
+			{
+				TargetDevice = Platforms[Index]->GetDevice(DeviceId);
+
+				if (TargetDevice.IsValid())
+				{
+					uint32 OutId;
+
+					if (TargetDevice->Run(Executable, Params, &OutId))
+					{
+						MessageSentTime = LastTime;
+						bMessageSent = true;
+						bLaunched = true;
+					}
+					else
+					{
+						MessageSentTime = LastTime;
+						bMessageSent = true;
+					}
+				}
+			}
+		}
+	}
+
+	return bLaunched;
+}
diff --git a/Engine/Source/Programs/UnrealFrontend/Private/Commands/LaunchCommand.h b/Engine/Source/Programs/UnrealFrontend/Private/Commands/LaunchCommand.h
new file mode 100644
--- /dev/null
+++ b/Engine/Source/Programs/UnrealFrontend/Private/Commands/LaunchCommand.h
@@ -0,0 +1,16 @@
+// Copyright 1998-2014 Epic Games, Inc. All Rights Reserved.
+
+#pragma once
+
+
+class FLaunchCommand
+{
+public:
+
+	/**
+	 * Executes the command.
+	 *
+	 * @return true on success, false otherwise.
+	 */
+	static bool Run( const FString& Params );
+};
diff --git a/Engine/Source/Programs/UnrealFrontend/Private/Commands/PackageCommand.cpp b/Engine/Source/Programs/UnrealFrontend/Private/Commands/PackageCommand.cpp
new file mode 100644
--- /dev/null
+++ b/Engine/Source/Programs/UnrealFrontend/Private/Commands/PackageCommand.cpp
@@ -0,0 +1,26 @@
+// Copyright 1998-2014 Epic Games, Inc. All Rights Reserved.
+
+#include "UnrealFrontendPrivatePCH.h"
+#include "PackageCommand.h"
+#include "TargetPlatform.h"
+
+
+void FPackageCommand::Run( )
+{
+	FString SourceDir;
+	FParse::Value(FCommandLine::Get(), TEXT("-SOURCEDIR="), SourceDir);
+
+	ITargetPlatformManagerModule* TPM = GetTargetPlatformManager();
+
+	if (TPM)
+	{
+		const TArray<ITargetPlatform*>& Platforms = TPM->GetActiveTargetPlatforms();
+
+		for (int32 Index = 0; Index < Platforms.Num(); ++Index)
+		{
+			if (Platforms[Index]->PackageBuild(SourceDir))
+			{
+			}
+		}
+	}
+}
diff --git a/Engine/Source/Programs/UnrealFrontend/Private/Commands/PackageCommand.h b/Engine/Source/Programs/UnrealFrontend/Private/Commands/PackageCommand.h
new file mode 100644
--- /dev/null
+++ b/Engine/Source/Programs/UnrealFrontend/Private/Commands/PackageCommand.h
@@ -0,0 +1,12 @@
+// Copyright 1998-2014 Epic Games, Inc. All Rights Reserved.
+
+#pragma once
+
+
+class FPackageCommand
+{
+public:
+
+	/** Executes the command. */
+	static void Run( );
+};
diff --git a/Engine/Source/Programs/UnrealFrontend/Private/Commands/StatsConvertCommand.cpp b/Engine/Source/Programs/UnrealFrontend/Private/Commands/StatsConvertCommand.cpp
new file mode 100644
--- /dev/null
+++ b/Engine/Source/Programs/UnrealFrontend/Private/Commands/StatsConvertCommand.cpp
@@ -0,0 +1,239 @@
+// Copyright 1998-2014 Epic Games, Inc. All Rights Reserved.
+
+#include "UnrealFrontendPrivatePCH.h"
+#include "StatsConvertCommand.h"
+
+
+void FStatsConvertCommand::WriteString( FArchive& Writer, const ANSICHAR* Format, ... )
+{
+	ANSICHAR Array[1024];
+	va_list ArgPtr;
+	va_start(ArgPtr,Format);
+	// Build the string
+	int32 Result = FCStringAnsi::GetVarArgs(Array,ARRAY_COUNT(Array),ARRAY_COUNT(Array)-1,Format,ArgPtr);
+	// Now write that to the file
+	Writer.Serialize((void*)Array,Result);
+}
+
+
+void FStatsConvertCommand::InternalRun()
+{
+#if	STATS
+	// get the target file
+	FString TargetFile;
+	FParse::Value(FCommandLine::Get(), TEXT("-INFILE="), TargetFile);
+	FString OutFile;
+	FParse::Value(FCommandLine::Get(), TEXT("-OUTFILE="), OutFile);
+	FString StatListString;
+	FParse::Value(FCommandLine::Get(), TEXT("-STATLIST="), StatListString);
+
+	// get the list of stats
+	TArray<FString> StatArrayString;
+	if (StatListString.ParseIntoArray(&StatArrayString, TEXT("+"), true) == 0)
+	{
+		StatArrayString.Add(TEXT("STAT_FrameTime"));
+	}
+
+	// convert to FNames for faster compare
+	for( const FString& It : StatArrayString )
+	{
+		StatList.Add(*It);
+	}
+
+	// open a csv file for write
+	TAutoPtr<FArchive> FileWriter( IFileManager::Get().CreateFileWriter( *OutFile ) );
+	if (!FileWriter)
+	{
+		UE_LOG( LogStats, Error, TEXT( "Could not open output file: %s" ), *OutFile );
+		return;
+	}
+
+	// @TODO yrx 2014-03-24 move to function
+	// attempt to read the data and convert to csv
+	const int64 Size = IFileManager::Get().FileSize( *TargetFile );
+	if( Size < 4 )
+	{
+		UE_LOG( LogStats, Error, TEXT( "Could not open input file: %s" ), *TargetFile );
+		return;
+	}
+
+	TAutoPtr<FArchive> FileReader( IFileManager::Get().CreateFileReader( *TargetFile ) );
+	if( !FileReader )
+	{
+		UE_LOG( LogStats, Error, TEXT( "Could not open input file: %s" ), *TargetFile );
+		return;
+	}
+
+	if( !Stream.ReadHeader( *FileReader ) )
+	{
+		UE_LOG( LogStats, Error, TEXT( "Could not open input file, bad magic: %s" ), *TargetFile );
+		return;
+	}
+
+	// This is not supported yet.
+	if (Stream.Header.bRawStatsFile)
+	{
+		UE_LOG( LogStats, Error, TEXT( "Could not open input file, not supported type (raw): %s" ), *TargetFile );
+		return;
+	}
+
+	const bool bIsFinalized = Stream.Header.IsFinalized();
+	if( bIsFinalized )
+	{
+		// Read metadata.
+		TArray<FStatMessage> MetadataMessages;
+		Stream.ReadFNamesAndMetadataMessages( *FileReader, MetadataMessages );
+		ThreadState.ProcessMetaDataOnly( MetadataMessages );
+
+		// Read frames offsets.
+		Stream.ReadFramesOffsets( *FileReader );
+		FileReader->Seek( Stream.FramesInfo[0].FrameFileOffset );
+	}
+
+	if( Stream.Header.HasCompressedData() )
+	{
+		UE_CLOG( !bIsFinalized, LogStats, Fatal, TEXT( "Compressed stats file has to be finalized" ) );
+	}
+
+	ReadAndConvertStatMessages( *FileReader, *FileWriter );
+
+#endif // STATS
+}
+
+
+void FStatsConvertCommand::ReadAndConvertStatMessages( FArchive& Reader, FArchive& Writer )
+{
+	// output the csv header
+	WriteString( Writer, "Frame,Name,Value\r\n" );
+
+	uint64 ReadMessages = 0;
+	TArray<FStatMessage> Messages;
+
+	// Buffer used to store the compressed and decompressed data.
+	TArray<uint8> SrcData;
+	TArray<uint8> DestData;
+
+	const bool bHasCompressedData = Stream.Header.HasCompressedData();
+	const bool bIsFinalized = Stream.Header.IsFinalized();
+	float DataLoadingProgress = 0.0f;
+
+	if( bHasCompressedData )
+	{
+		while( Reader.Tell() < Reader.TotalSize() )
+		{
+			// Read the compressed data.
+			FCompressedStatsData UncompressedData( SrcData, DestData );
+			Reader << UncompressedData;
+			if( UncompressedData.HasReachedEndOfCompressedData() )
+			{
+				return;
+			}
+
+			FMemoryReader MemoryReader( DestData, true );
+
+			while( MemoryReader.Tell() < MemoryReader.TotalSize() )
+			{
+				// read the message
+				FStatMessage Message( Stream.ReadMessage( MemoryReader, bIsFinalized ) );
+				ReadMessages++;
+				if( ReadMessages % 32768 == 0 )
+				{
+					UE_LOG( LogStats, Log, TEXT( "StatsConvertCommand progress: %.1f%%" ), DataLoadingProgress );
+				}
+
+				if( Message.NameAndInfo.GetShortName() != TEXT( "Unknown FName" ) )
+				{
+					if( Message.NameAndInfo.GetField<EStatOperation>() == EStatOperation::AdvanceFrameEventGameThread && ReadMessages > 2 )
+					{
+						new (Messages) FStatMessage( Message );
+						ThreadState.AddMessages( Messages );
+						Messages.Reset();
+
+						CollectAndWriteStatsValues( Writer );
+						DataLoadingProgress = (double)Reader.Tell() / (double)Reader.TotalSize() * 100.0f;
+					}
+
+					new (Messages) FStatMessage( Message );
+				}
+				else
+				{
+					break;
+				}
+			}
+		}
+	}
+	else
+	{
+		while( Reader.Tell() < Reader.TotalSize() )
+		{
+			// read the message
+			FStatMessage Message( Stream.ReadMessage( Reader, bIsFinalized ) );
+			ReadMessages++;
+			if( ReadMessages % 32768 == 0 )
+			{
+				UE_LOG( LogStats, Log, TEXT( "StatsConvertCommand progress: %.1f%%" ), DataLoadingProgress );
+			}
+
+			if( Message.NameAndInfo.GetShortName() != TEXT( "Unknown FName" ) )
+			{
+				if( Message.NameAndInfo.GetField<EStatOperation>() == EStatOperation::SpecialMessageMarker )
+				{
+					// Simply break the loop.
+					// The profiler supports more advanced handling of this message.
+					return;
+				}
+				else if( Message.NameAndInfo.GetField<EStatOperation>() == EStatOperation::AdvanceFrameEventGameThread && ReadMessages > 2 )
+				{
+					new (Messages) FStatMessage( Message );
+					ThreadState.AddMessages( Messages );
+					Messages.Reset();
+
+					CollectAndWriteStatsValues( Writer );
+					DataLoadingProgress = (double)Reader.Tell() / (double)Reader.TotalSize() * 100.0f;
+				}
+
+				new (Messages) FStatMessage( Message );
+			}
+			else
+			{
+				break;
+			}
+		}
+	}
+}
+
+
+void FStatsConvertCommand::CollectAndWriteStatsValues( FArchive& Writer )
+{
+	// get the thread stats
+	TArray<FStatMessage> Stats;
+	ThreadState.GetInclusiveAggregateStackStats(ThreadState.CurrentGameFrame, Stats);
+	for (int32 Index = 0; Index < Stats.Num(); ++Index)
+	{
+		FStatMessage const& Meta = Stats[Index];
+		//UE_LOG(LogTemp, Display, TEXT("Stat: %s"), *Meta.NameAndInfo.GetShortName().ToString());
+
+		for (int32 Jndex = 0; Jndex < StatList.Num(); ++Jndex)
+		{
+			if (Meta.NameAndInfo.GetShortName() == StatList[Jndex])
+			{
+				double StatValue = 0.0f;
+				if (Meta.NameAndInfo.GetFlag(EStatMetaFlags::IsPackedCCAndDuration))
+				{
+					StatValue = FPlatformTime::ToMilliseconds( FromPackedCallCountDuration_Duration(Meta.GetValue_int64()) );
+				}
+				else if (Meta.NameAndInfo.GetFlag(EStatMetaFlags::IsCycle))
+				{
+					StatValue = FPlatformTime::ToMilliseconds( Meta.GetValue_int64() );
+				}
+				else
+				{
+					StatValue = Meta.GetValue_int64();
+				}
+
+				// write out to the csv file
+				WriteString( Writer, "%d,%S,%f\r\n", ThreadState.CurrentGameFrame, *StatList[Jndex].ToString(), StatValue );
+			}
+		}
+	}
+}
diff --git a/Engine/Source/Programs/UnrealFrontend/Private/Commands/StatsConvertCommand.h b/Engine/Source/Programs/UnrealFrontend/Private/Commands/StatsConvertCommand.h
new file mode 100644
--- /dev/null
+++ b/Engine/Source/Programs/UnrealFrontend/Private/Commands/StatsConvertCommand.h
@@ -0,0 +1,32 @@
+// Copyright 1998-2014 Epic Games, Inc. All Rights Reserved.
+
+#pragma once
+
+#include "StatsData.h"
+#include "StatsFile.h"
+
+
+class FStatsConvertCommand
+{
+public:
+
+	/** Executes the command. */
+	static void Run()
+	{
+		FStatsConvertCommand Instance;
+		Instance.InternalRun();
+	}
+
+protected:
+
+	void InternalRun();
+	void WriteString( FArchive& Writer, const ANSICHAR* Format, ... );
+	void CollectAndWriteStatsValues( FArchive& Writer );
+	void ReadAndConvertStatMessages( FArchive& Reader, FArchive& Writer );
+
+private:
+
+	FStatsReadStream Stream;
+	FStatsThreadState ThreadState;
+	TArray<FName> StatList;
+};
diff --git a/Engine/Source/Programs/UnrealFrontend/Private/Commands/UserInterfaceCommand.cpp b/Engine/Source/Programs/UnrealFrontend/Private/Commands/UserInterfaceCommand.cpp
new file mode 100644
--- /dev/null
+++ b/Engine/Source/Programs/UnrealFrontend/Private/Commands/UserInterfaceCommand.cpp
@@ -0,0 +1,143 @@
+// Copyright 1998-2014 Epic Games, Inc. All Rights Reserved.
+
+#include "UnrealFrontendPrivatePCH.h"
+#include "AutomationController.h"
+#include "Projects.h"
+#include "Slate.h"
+#include "SlateReflector.h"
+#include "StandaloneRenderer.h"
+#include "TaskGraphInterfaces.h"
+#include "UserInterfaceCommand.h"
+
+
+#define IDEAL_FRAMERATE 60;
+
+
+namespace UserInterfaceCommand
+{
+	TSharedPtr<FTabManager::FLayout> ApplicationLayout;
+	TSharedRef<FWorkspaceItem> DeveloperTools = FWorkspaceItem::NewGroup( NSLOCTEXT("UnrealFrontend", "DeveloperToolsMenu", "Developer Tools") );
+}
+
+
+/* FUserInterfaceCommand interface
+ *****************************************************************************/
+
+void FUserInterfaceCommand::Run(  )
+{
+	FString UnrealFrontendLayoutIni = FPaths::GetPath(GEngineIni) + "/Layout.ini";
+
+	// load required modules
+	FModuleManager::Get().LoadModuleChecked("EditorStyle");
+	FModuleManager::Get().LoadModuleChecked("Messaging");
+
+	IAutomationControllerModule& AutomationControllerModule = FModuleManager::LoadModuleChecked<IAutomationControllerModule>("AutomationController");
+	AutomationControllerModule.Init();
+
+	// load plug-ins
+	// @todo: allow for better plug-in support in standalone Slate apps
+	IPluginManager::Get().LoadModulesForEnabledPlugins(ELoadingPhase::PreDefault);
+
+	// load optional modules
+	FModuleManager::Get().LoadModule("DeviceManager");
+	FModuleManager::Get().LoadModule("ProfilerClient");
+	FModuleManager::Get().LoadModule("ProjectLauncher");
+	FModuleManager::Get().LoadModule("SessionFrontend");
+	FModuleManager::Get().LoadModule("SettingsEditor");
+
+	InitializeSlateApplication(UnrealFrontendLayoutIni);
+
+	// enter main loop
+	double DeltaTime = 0.0;
+	double LastTime = FPlatformTime::Seconds();
+	const float IdealFrameTime = 1.0f / IDEAL_FRAMERATE;
+
+	while (!GIsRequestingExit)
+	{
+		//Save the state of the tabs here rather than after close of application (the tabs are undesirably saved out with ClosedTab state on application close)
+		//UserInterfaceCommand::UserConfiguredNewLayout = FGlobalTabmanager::Get()->PersistLayout();
+
+		FTaskGraphInterface::Get().ProcessThreadUntilIdle(ENamedThreads::GameThread);
+
+		FSlateApplication::Get().PumpMessages();
+		FSlateApplication::Get().Tick();
+		FTicker::GetCoreTicker().Tick(DeltaTime);
+		AutomationControllerModule.Tick();
+
+		// throttle frame rate
+		FPlatformProcess::Sleep(FMath::Max<float>(0.0f, IdealFrameTime - (FPlatformTime::Seconds() - LastTime)));
+
+		double CurrentTime = FPlatformTime::Seconds();
+		DeltaTime =  CurrentTime - LastTime;
+		LastTime = CurrentTime;
+
+		FStats::AdvanceFrame( false );
+
+		GLog->FlushThreadedLogs();
+	}
+
+	ShutdownSlateApplication(UnrealFrontendLayoutIni);
+}
+
+
+/* FUserInterfaceCommand implementation
+ *****************************************************************************/
+
+void FUserInterfaceCommand::InitializeSlateApplication( const FString& LayoutIni )
+{
+	FSlateApplication::InitializeAsStandaloneApplication(GetStandardStandaloneRenderer());
+	FGlobalTabmanager::Get()->SetApplicationTitle(NSLOCTEXT("UnrealFrontend", "AppTitle", "Unreal Frontend"));
+
+	// load widget reflector
+	const bool bAllowDebugTools = FParse::Param(FCommandLine::Get(), TEXT("DebugTools"));
+
+	if (bAllowDebugTools)
+	{
+		ISlateReflectorModule* SlateReflectorModule = FModuleManager::GetModulePtr<ISlateReflectorModule>("SlateReflector");
+
+		if (SlateReflectorModule != nullptr)
+		{
+			SlateReflectorModule->RegisterTabSpawner(UserInterfaceCommand::DeveloperTools);
+		}
+	}
+
+	// restore application layout
+	TSharedRef<FTabManager::FLayout> NewLayout = FTabManager::NewLayout("SessionFrontendLayout_v1.1")
+		->AddArea
+		(
+			FTabManager::NewArea(1280.f, 720.0f)
+				->Split
+				(
+					FTabManager::NewStack()
+						->AddTab(FName("DeviceManager"), ETabState::OpenedTab)
+						->AddTab(FName("MessagingDebugger"), ETabState::ClosedTab)
+						->AddTab(FName("SessionFrontend"), ETabState::OpenedTab)
+						->AddTab(FName("ProjectLauncher"), ETabState::OpenedTab)
+				)
+		)
+		->AddArea
+		(
+			FTabManager::NewArea(600.0f, 600.0f)
+				->SetWindow(FVector2D(10.0f, 10.0f), false)
+				->Split
+				(
+					FTabManager::NewStack()->AddTab("WidgetReflector", bAllowDebugTools ? ETabState::OpenedTab : ETabState::ClosedTab)
+				)
+		);
+
+	UserInterfaceCommand::ApplicationLayout = FLayoutSaveRestore::LoadFromConfig(LayoutIni, NewLayout);
+	FGlobalTabmanager::Get()->RestoreFrom(UserInterfaceCommand::ApplicationLayout.ToSharedRef(), TSharedPtr<SWindow>());
+}
+
+
+void FUserInterfaceCommand::ShutdownSlateApplication( const FString& LayoutIni )
+{
+	check(UserInterfaceCommand::ApplicationLayout.IsValid());
+
+	// save application layout
+	FLayoutSaveRestore::SaveToConfig(LayoutIni, UserInterfaceCommand::ApplicationLayout.ToSharedRef());
+	GConfig->Flush(false, LayoutIni);
+
+	// shut down application
+	FSlateApplication::Shutdown();
+}
diff --git a/Engine/Source/Programs/UnrealFrontend/Private/Commands/UserInterfaceCommand.h b/Engine/Source/Programs/UnrealFrontend/Private/Commands/UserInterfaceCommand.h
new file mode 100644
--- /dev/null
+++ b/Engine/Source/Programs/UnrealFrontend/Private/Commands/UserInterfaceCommand.h
@@ -0,0 +1,28 @@
+// Copyright 1998-2014 Epic Games, Inc. All Rights Reserved.
+
+#pragma once
+
+
+class FUserInterfaceCommand
+{
+public:
+
+	/** Executes the command. */
+	static void Run( );
+
+protected:
+
+	/**
+	 * Initializes the Slate application.
+	 *
+	 * @param LayoutInit The path to the layout configuration file.
+	 */
+	static void InitializeSlateApplication( const FString& LayoutIni );
+
+	/**
+	 * Shuts down the Slate application.
+	 *
+	 * @param LayoutInit The path to the layout configuration file.
+	 */
+	static void ShutdownSlateApplication( const FString& LayoutIni );
+};
diff --git a/Engine/Source/Programs/UnrealFrontend/Private/Linux/LinuxUnrealFrontendMain.cpp b/Engine/Source/Programs/UnrealFrontend/Private/Linux/LinuxUnrealFrontendMain.cpp
--- a/Engine/Source/Programs/UnrealFrontend/Private/Linux/LinuxUnrealFrontendMain.cpp
+++ b/Engine/Source/Programs/UnrealFrontend/Private/Linux/LinuxUnrealFrontendMain.cpp
@@ -1,5 +1,6 @@
 // Copyright 1998-2014 Epic Games, Inc. All Rights Reserved.
 
+#include "UnrealFrontendPrivatePCH.h"
 #include "UnrealFrontendMain.h"
 
 
@@ -12,12 +13,14 @@ int main( int argc, char *argv[] )
 	{
 		GSavedCommandLine += TEXT(" ");
 		FString Argument(ANSI_TO_TCHAR(argv[Option]));
+
 		if (Argument.Contains(TEXT(" ")))
 		{
 			if (Argument.Contains(TEXT("=")))
 			{
 				FString ArgName;
 				FString ArgValue;
+
 				Argument.Split( TEXT("="), &ArgName, &ArgValue );
 				Argument = FString::Printf( TEXT("%s=\"%s\""), *ArgName, *ArgValue );
 			}
@@ -26,6 +29,7 @@ int main( int argc, char *argv[] )
 				Argument = FString::Printf(TEXT("\"%s\""), *Argument);
 			}
 		}
+
 		GSavedCommandLine += Argument;
 	}
 	
diff --git a/Engine/Source/Programs/UnrealFrontend/Private/Mac/MacUnrealFrontendMain.cpp b/Engine/Source/Programs/UnrealFrontend/Private/Mac/MacUnrealFrontendMain.cpp
--- a/Engine/Source/Programs/UnrealFrontend/Private/Mac/MacUnrealFrontendMain.cpp
+++ b/Engine/Source/Programs/UnrealFrontend/Private/Mac/MacUnrealFrontendMain.cpp
@@ -1,5 +1,6 @@
 // Copyright 1998-2014 Epic Games, Inc. All Rights Reserved.
 
+#include "UnrealFrontendPrivatePCH.h"
 #include "UnrealFrontendMain.h"
 #include "ExceptionHandling.h"
 #include "RequiredProgramMainCPPInclude.h"
diff --git a/Engine/Source/Programs/UnrealFrontend/Private/UnrealFrontendMain.cpp b/Engine/Source/Programs/UnrealFrontend/Private/UnrealFrontendMain.cpp
--- a/Engine/Source/Programs/UnrealFrontend/Private/UnrealFrontendMain.cpp
+++ b/Engine/Source/Programs/UnrealFrontend/Private/UnrealFrontendMain.cpp
@@ -1,577 +1,40 @@
 // Copyright 1998-2014 Epic Games, Inc. All Rights Reserved.
 
+#include "UnrealFrontendPrivatePCH.h"
 #include "UnrealFrontendMain.h"
-
 #include "RequiredProgramMainCPPInclude.h"
-#include "AutomationController.h"
-#include "DeviceManager.h"
-#include "ProfilerClient.h"
-#include "SessionFrontend.h"
-#include "StatsData.h"
-#include "StatsFile.h"
-#include "EditorStyle.h"
-#include "SlateReflector.h"
-
-
-IMPLEMENT_APPLICATION(UnrealFrontend, "UnrealFrontend");
-
-#define IDEAL_FRAMERATE 60;
-
-
-namespace WorkspaceMenu
-{
-	TSharedRef<FWorkspaceItem> DeveloperTools = FWorkspaceItem::NewGroup( NSLOCTEXT("UnrealFrontend", "DeveloperToolsMenu", "Developer Tools") );
-}
-
-
-void RunPackageCommand()
-{
-	FString SourceDir;
-	FParse::Value(FCommandLine::Get(), TEXT("-SOURCEDIR="), SourceDir);
-
-	ITargetPlatformManagerModule* TPM = GetTargetPlatformManager();
-	if (TPM)
-	{
-		const TArray<ITargetPlatform*>& Platforms = TPM->GetActiveTargetPlatforms();
-
-		for (int32 Index = 0; Index < Platforms.Num(); ++Index)
-		{
-			if (Platforms[Index]->PackageBuild(SourceDir))
-			{
-			}
-		}
-	}
-}
-
-
-bool RunDeployCommand()
-{
-	bool bDeployed = false;
-
-	// get the target device
-	FString Device;
-	FParse::Value(FCommandLine::Get(), TEXT("-DEVICE="), Device);
-
-	// get the file manifest
-	FString Manifest;
-	FParse::Value(FCommandLine::Get(), TEXT("-MANIFEST="), Manifest);
-
-	FString SourceDir;
-	FParse::Value(FCommandLine::Get(), TEXT("-SOURCEDIR="), SourceDir);
-
-	ITargetPlatformManagerModule* TPM = GetTargetPlatformManager();
-	if (!TPM)
-	{
-		return false;
-	}
-
-	// Initialize the messaging subsystem so we can do device discovery.
-	FModuleManager::LoadModuleChecked<IMessagingModule>("Messaging");
-
-	// load plug-in modules
-	// @todo: allow for better plug-in support in standalone Slate apps
-	IPluginManager::Get().LoadModulesForEnabledPlugins(ELoadingPhase::PreDefault);
-
-	double DeltaTime = 0.0;
-	double LastTime = FPlatformTime::Seconds();
-
-	// We track the message sent time because we have to keep updating the loop until the message is *actually sent*. (ie all packets queued, sent, buffer flushed, etc.)
-	double MessageSentTime = 0.0;
-	bool bMessageSent = false;
-	while ( !GIsRequestingExit && (MessageSentTime > LastTime + 1.0 || MessageSentTime <= 0.1) )
-	{
-		FTaskGraphInterface::Get().ProcessThreadUntilIdle(ENamedThreads::GameThread);
-		FTicker::GetCoreTicker().Tick(DeltaTime);
-		FPlatformProcess::Sleep(0);
-
-		DeltaTime = FPlatformTime::Seconds() - LastTime;
-		LastTime = FPlatformTime::Seconds();
-
-		if ( !bMessageSent )
-		{
-			const TArray<ITargetPlatform*>& Platforms = TPM->GetActiveTargetPlatforms();
-
-			FString Platform;
-			FString DeviceName;
-			Device.Split(TEXT("@"), &Platform, &DeviceName);
-			FTargetDeviceId DeviceId(Platform, DeviceName);
-			ITargetDevicePtr TargetDevice;
-			for (int32 Index = 0; Index < Platforms.Num(); ++Index)
-			{
-				TargetDevice = Platforms[Index]->GetDevice(DeviceId);
-				if (TargetDevice.IsValid())
-				{
-					FString OutId;
-					if (Platforms[Index]->PackageBuild(SourceDir))
-					{
-						if (TargetDevice->Deploy(SourceDir, OutId))
-						{
-							bDeployed = true;
-						}
-						MessageSentTime = LastTime;
-						bMessageSent = true;
-					}
-					else
-					{
-						MessageSentTime = LastTime;
-						bMessageSent = true;
-					}
-				}
-			}
-		}
-	}
-	return bDeployed;
-}
-
-
-bool RunLaunchCommand(const FString& Params)
-{
-	bool bLaunched = false;
-
-	// get the target device
-	FString Device;
-	FParse::Value(FCommandLine::Get(), TEXT("-DEVICE="), Device);
-
-	// get the executable to launch
-	FString Executable;
-	FParse::Value(FCommandLine::Get(), TEXT("-EXE="), Executable);
-
-	ITargetPlatformManagerModule* TPM = GetTargetPlatformManager();
-	if (!TPM)
-	{
-		return false;
-	}
-
-	// Initialize the messaging subsystem so we can do device discovery.
-	FModuleManager::LoadModuleChecked<IMessagingModule>("Messaging");
-
-	// load plug-in modules
-	// @todo: allow for better plug-in support in standalone Slate apps
-	IPluginManager::Get().LoadModulesForEnabledPlugins(ELoadingPhase::PreDefault);
-
-	double DeltaTime = 0.0;
-	double LastTime = FPlatformTime::Seconds();
-	static int32 MasterDisableChangeTagStartFrame = -1;
-
-	// We track the message sent time because we have to keep updating the loop until the message is *actually sent*. (ie all packets queued, sent, buffer flushed, etc.)
-	double MessageSentTime = 0.0;
-	bool bMessageSent = false;
-	while ( !GIsRequestingExit && (MessageSentTime > LastTime + 1.0 || MessageSentTime <= 0.1) )
-	{
-		FTaskGraphInterface::Get().ProcessThreadUntilIdle(ENamedThreads::GameThread);
-		FTicker::GetCoreTicker().Tick(DeltaTime);
-		FPlatformProcess::Sleep(0);
-
-		DeltaTime = FPlatformTime::Seconds() - LastTime;
-		LastTime = FPlatformTime::Seconds();
-
-		if ( !bMessageSent )
-		{
-			const TArray<ITargetPlatform*>& Platforms = TPM->GetActiveTargetPlatforms();
-
-			FString Platform;
-			FString DeviceName;
-			Device.Split(TEXT("@"), &Platform, &DeviceName);
-			FTargetDeviceId DeviceId(Platform, DeviceName);
-			ITargetDevicePtr TargetDevice;
-			for (int32 Index = 0; Index < Platforms.Num(); ++Index)
-			{
-				TargetDevice = Platforms[Index]->GetDevice(DeviceId);
-				if (TargetDevice.IsValid())
-				{
-					uint32 OutId;
-					if (TargetDevice->Run(Executable, Params, &OutId))
-					{
-						MessageSentTime = LastTime;
-						bMessageSent = true;
-						bLaunched = true;
-					}
-					else
-					{
-						MessageSentTime = LastTime;
-						bMessageSent = true;
-					}
-				}
-			}
-		}
-	}
-
-	return bLaunched;
-}
-
-
-struct FStatsConvertCommand
-{
-	static void Run()
-	{
-		FStatsConvertCommand Instance;
-		Instance.InternalRun();
-	}
-
-protected:
-	void InternalRun();
-	void WriteString( FArchive& Writer, const ANSICHAR* Format, ... );
-	void CollectAndWriteStatsValues( FArchive& Writer );
-	void ReadAndConvertStatMessages( FArchive& Reader, FArchive& Writer );
-
-	FStatsReadStream Stream;
-	FStatsThreadState ThreadState;
-	TArray<FName> StatList;
-};
-
-void FStatsConvertCommand::WriteString( FArchive& Writer, const ANSICHAR* Format, ... )
-{
-	ANSICHAR Array[1024];
-	va_list ArgPtr;
-	va_start(ArgPtr,Format);
-	// Build the string
-	int32 Result = FCStringAnsi::GetVarArgs(Array,ARRAY_COUNT(Array),ARRAY_COUNT(Array)-1,Format,ArgPtr);
-	// Now write that to the file
-	Writer.Serialize((void*)Array,Result);
-}
-
-void FStatsConvertCommand::InternalRun()
-{
-#if	STATS
-	// get the target file
-	FString TargetFile;
-	FParse::Value(FCommandLine::Get(), TEXT("-INFILE="), TargetFile);
-	FString OutFile;
-	FParse::Value(FCommandLine::Get(), TEXT("-OUTFILE="), OutFile);
-	FString StatListString;
-	FParse::Value(FCommandLine::Get(), TEXT("-STATLIST="), StatListString);
-
-	// get the list of stats
-	TArray<FString> StatArrayString;
-	if (StatListString.ParseIntoArray(&StatArrayString, TEXT("+"), true) == 0)
-	{
-		StatArrayString.Add(TEXT("STAT_FrameTime"));
-	}
-
-	// convert to FNames for faster compare
-	for( const FString& It : StatArrayString )
-	{
-		StatList.Add(*It);
-	}
-
-	// open a csv file for write
-	TAutoPtr<FArchive> FileWriter( IFileManager::Get().CreateFileWriter( *OutFile ) );
-	if (!FileWriter)
-	{
-		UE_LOG( LogStats, Error, TEXT( "Could not open output file: %s" ), *OutFile );
-		return;
-	}
-
-	// @TODO yrx 2014-03-24 move to function
-	// attempt to read the data and convert to csv
-	const int64 Size = IFileManager::Get().FileSize( *TargetFile );
-	if( Size < 4 )
-	{
-		UE_LOG( LogStats, Error, TEXT( "Could not open input file: %s" ), *TargetFile );
-		return;
-	}
-
-	TAutoPtr<FArchive> FileReader( IFileManager::Get().CreateFileReader( *TargetFile ) );
-	if( !FileReader )
-	{
-		UE_LOG( LogStats, Error, TEXT( "Could not open input file: %s" ), *TargetFile );
-		return;
-	}
-
-	if( !Stream.ReadHeader( *FileReader ) )
-	{
-		UE_LOG( LogStats, Error, TEXT( "Could not open input file, bad magic: %s" ), *TargetFile );
-		return;
-	}
-
-	// This is not supported yet.
-	if (Stream.Header.bRawStatsFile)
-	{
-		UE_LOG( LogStats, Error, TEXT( "Could not open input file, not supported type (raw): %s" ), *TargetFile );
-		return;
-	}
-
-	const bool bIsFinalized = Stream.Header.IsFinalized();
-	if( bIsFinalized )
-	{
-		// Read metadata.
-		TArray<FStatMessage> MetadataMessages;
-		Stream.ReadFNamesAndMetadataMessages( *FileReader, MetadataMessages );
-		ThreadState.ProcessMetaDataOnly( MetadataMessages );
-
-		// Read frames offsets.
-		Stream.ReadFramesOffsets( *FileReader );
-		FileReader->Seek( Stream.FramesInfo[0].FrameFileOffset );
-	}
-
-	if( Stream.Header.HasCompressedData() )
-	{
-		UE_CLOG( !bIsFinalized, LogStats, Fatal, TEXT( "Compressed stats file has to be finalized" ) );
-	}
-
-	ReadAndConvertStatMessages( *FileReader, *FileWriter );
-
-#endif // STATS
-}
 
-void FStatsConvertCommand::ReadAndConvertStatMessages( FArchive& Reader, FArchive& Writer )
-{
-	// output the csv header
-	WriteString( Writer, "Frame,Name,Value\r\n" );
-
-	uint64 ReadMessages = 0;
-	TArray<FStatMessage> Messages;
-
-	// Buffer used to store the compressed and decompressed data.
-	TArray<uint8> SrcData;
-	TArray<uint8> DestData;
+#include "DeployCommand.h"
+#include "LaunchCommand.h"
+#include "PackageCommand.h"
+#include "StatsConvertCommand.h"
+#include "UserInterfaceCommand.h"
 
-	const bool bHasCompressedData = Stream.Header.HasCompressedData();
-	const bool bIsFinalized = Stream.Header.IsFinalized();
-	float DataLoadingProgress = 0.0f;
-
-	if( bHasCompressedData )
-	{
-		while( Reader.Tell() < Reader.TotalSize() )
-		{
-			// Read the compressed data.
-			FCompressedStatsData UncompressedData( SrcData, DestData );
-			Reader << UncompressedData;
-			if( UncompressedData.HasReachedEndOfCompressedData() )
-			{
-				return;
-			}
 
-			FMemoryReader MemoryReader( DestData, true );
-
-			while( MemoryReader.Tell() < MemoryReader.TotalSize() )
-			{
-				// read the message
-				FStatMessage Message( Stream.ReadMessage( MemoryReader, bIsFinalized ) );
-				ReadMessages++;
-				if( ReadMessages % 32768 == 0 )
-				{
-					UE_LOG( LogStats, Log, TEXT( "StatsConvertCommand progress: %.1f%%" ), DataLoadingProgress );
-				}
-
-				if( Message.NameAndInfo.GetShortName() != TEXT( "Unknown FName" ) )
-				{
-					if( Message.NameAndInfo.GetField<EStatOperation>() == EStatOperation::AdvanceFrameEventGameThread && ReadMessages > 2 )
-					{
-						new (Messages) FStatMessage( Message );
-						ThreadState.AddMessages( Messages );
-						Messages.Reset();
-
-						CollectAndWriteStatsValues( Writer );
-						DataLoadingProgress = (double)Reader.Tell() / (double)Reader.TotalSize() * 100.0f;
-					}
-
-					new (Messages) FStatMessage( Message );
-				}
-				else
-				{
-					break;
-				}
-			}
-		}
-	}
-	else
-	{
-		while( Reader.Tell() < Reader.TotalSize() )
-		{
-			// read the message
-			FStatMessage Message( Stream.ReadMessage( Reader, bIsFinalized ) );
-			ReadMessages++;
-			if( ReadMessages % 32768 == 0 )
-			{
-				UE_LOG( LogStats, Log, TEXT( "StatsConvertCommand progress: %.1f%%" ), DataLoadingProgress );
-			}
-
-			if( Message.NameAndInfo.GetShortName() != TEXT( "Unknown FName" ) )
-			{
-				if( Message.NameAndInfo.GetField<EStatOperation>() == EStatOperation::SpecialMessageMarker )
-				{
-					// Simply break the loop.
-					// The profiler supports more advanced handling of this message.
-					return;
-				}
-				else if( Message.NameAndInfo.GetField<EStatOperation>() == EStatOperation::AdvanceFrameEventGameThread && ReadMessages > 2 )
-				{
-					new (Messages) FStatMessage( Message );
-					ThreadState.AddMessages( Messages );
-					Messages.Reset();
-
-					CollectAndWriteStatsValues( Writer );
-					DataLoadingProgress = (double)Reader.Tell() / (double)Reader.TotalSize() * 100.0f;
-				}
-
-				new (Messages) FStatMessage( Message );
-			}
-			else
-			{
-				break;
-			}
-		}
-	}
-}
-
-void FStatsConvertCommand::CollectAndWriteStatsValues( FArchive& Writer )
-{
-	// get the thread stats
-	TArray<FStatMessage> Stats;
-	ThreadState.GetInclusiveAggregateStackStats(ThreadState.CurrentGameFrame, Stats);
-	for (int32 Index = 0; Index < Stats.Num(); ++Index)
-	{
-		FStatMessage const& Meta = Stats[Index];
-		//UE_LOG(LogTemp, Display, TEXT("Stat: %s"), *Meta.NameAndInfo.GetShortName().ToString());
-
-		for (int32 Jndex = 0; Jndex < StatList.Num(); ++Jndex)
-		{
-			if (Meta.NameAndInfo.GetShortName() == StatList[Jndex])
-			{
-				double StatValue = 0.0f;
-				if (Meta.NameAndInfo.GetFlag(EStatMetaFlags::IsPackedCCAndDuration))
-				{
-					StatValue = FPlatformTime::ToMilliseconds( FromPackedCallCountDuration_Duration(Meta.GetValue_int64()) );
-				}
-				else if (Meta.NameAndInfo.GetFlag(EStatMetaFlags::IsCycle))
-				{
-					StatValue = FPlatformTime::ToMilliseconds( Meta.GetValue_int64() );
-				}
-				else
-				{
-					StatValue = Meta.GetValue_int64();
-				}
-
-				// write out to the csv file
-				WriteString( Writer, "%d,%S,%f\r\n", ThreadState.CurrentGameFrame, *StatList[Jndex].ToString(), StatValue );
-			}
-		}
-	}
-}
-
-void RunUI()
-{
-	FString UnrealFrontendLayoutIni = FPaths::GetPath(GEngineIni) + "/Layout.ini";
-
-	FSlateApplication::InitializeAsStandaloneApplication(GetStandardStandaloneRenderer());
-
-	// The frontend currently relies on EditorStyle being loaded
-	FModuleManager::LoadModuleChecked<IEditorStyleModule>("EditorStyle");
-
-	// initialize messaging subsystem
-	FModuleManager::LoadModuleChecked<IMessagingModule>("Messaging");
-
-	// load plug-in modules
-	// @todo: allow for better plug-in support in standalone Slate apps
-	IPluginManager::Get().LoadModulesForEnabledPlugins(ELoadingPhase::PreDefault);
-
-	// initialize automation tests
-	IAutomationControllerModule& AutomationControllerModule = FModuleManager::LoadModuleChecked<IAutomationControllerModule>("AutomationController");
-	AutomationControllerModule.Init();
-
-	// initialize profiler
-	FModuleManager::Get().LoadModule("ProfilerClient");
-
-	// initialize user interface
-	FModuleManager::Get().LoadModule("DeviceManager");
-	FModuleManager::Get().LoadModule("SessionFrontend");
-	FModuleManager::Get().LoadModule("SettingsEditor");
-	FModuleManager::Get().LoadModule("ProjectLauncher");
-
-	// Create developer tools menu with widget reflector.
-	FModuleManager::LoadModuleChecked<ISlateReflectorModule>("SlateReflector").RegisterTabSpawner(WorkspaceMenu::DeveloperTools);
-
-	// check to see if want the widget reflector
-	const bool bAllowDebugTools = FParse::Param(FCommandLine::Get(), TEXT("DebugTools"));
-
-	// set the application name
-	FGlobalTabmanager::Get()->SetApplicationTitle(NSLOCTEXT("UnrealFrontend", "AppTitle", "Unreal Frontend"));
-
-	// restore application layout
-	TSharedRef<FTabManager::FLayout> NewLayout = FTabManager::NewLayout("SessionFrontendLayout_v1.1")
-		->AddArea
-		(
-			FTabManager::NewArea(1280.f, 720.0f)
-				->Split
-				(
-					FTabManager::NewStack()
-						->AddTab(FName("DeviceManager"), ETabState::OpenedTab)
-						->AddTab(FName("MessagingDebugger"), ETabState::ClosedTab)
-						->AddTab(FName("SessionFrontend"), ETabState::OpenedTab)
-						->AddTab(FName("ProjectLauncher"), ETabState::OpenedTab)
-				)
-		)
-		->AddArea
-		(
-			FTabManager::NewArea(600.0f, 600.0f)
-				->SetWindow(FVector2D(10.0f, 10.0f), false)
-				->Split
-				(
-					FTabManager::NewStack()->AddTab("WidgetReflector", bAllowDebugTools ? ETabState::OpenedTab : ETabState::ClosedTab)
-				)
-		);
-
-	TSharedRef<FTabManager::FLayout> UserConfiguredNewLayout = FLayoutSaveRestore::LoadFromConfig(UnrealFrontendLayoutIni, NewLayout);
-	FGlobalTabmanager::Get()->RestoreFrom(UserConfiguredNewLayout, TSharedPtr<SWindow>());
-
-	// enter main loop
-	double DeltaTime = 0.0;
-	double LastTime = FPlatformTime::Seconds();
-	const float IdealFrameTime = 1.0f / IDEAL_FRAMERATE;
-
-	while (!GIsRequestingExit)
-	{
-		//Save the state of the tabs here rather than after close of application (the tabs are undesirably saved out with ClosedTab state on application close)
-		//UserConfiguredNewLayout = FGlobalTabmanager::Get()->PersistLayout();
-
-		FTaskGraphInterface::Get().ProcessThreadUntilIdle(ENamedThreads::GameThread);
-
-		FSlateApplication::Get().PumpMessages();
-		FSlateApplication::Get().Tick();
-		FTicker::GetCoreTicker().Tick(DeltaTime);
-		AutomationControllerModule.Tick();
-
-		// throttle frame rate
-		FPlatformProcess::Sleep(FMath::Max<float>(0.0f, IdealFrameTime - (FPlatformTime::Seconds() - LastTime)));
-
-		double CurrentTime = FPlatformTime::Seconds();
-		DeltaTime =  CurrentTime - LastTime;
-		LastTime = CurrentTime;
-
-		FStats::AdvanceFrame( false );
-
-		GLog->FlushThreadedLogs();
-	}
-
-	// save application layout
-	FLayoutSaveRestore::SaveToConfig(UnrealFrontendLayoutIni, UserConfiguredNewLayout);
-	GConfig->Flush(false, UnrealFrontendLayoutIni);
-
-	// shut down application
-	FSlateApplication::Shutdown();
-}
+IMPLEMENT_APPLICATION(UnrealFrontend, "UnrealFrontend");
 
 
+/**
+ * Platform agnostic implementation of the main entry point.
+ */
 int32 UnrealFrontendMain( const TCHAR* CommandLine )
 {
-	// check to see if we should run in command line mode
 	FString Command;
 	FString Params;
 	FString NewCommandLine = CommandLine;
 
+	// process command line parameters
 	bool bRunCommand = FParse::Value(*NewCommandLine, TEXT("-RUN="), Command);
 	
 	if (bRunCommand)
 	{
 		// extract off any -PARAM= parameters so that they aren't accidentally parsed by engine init
 		FParse::Value(*NewCommandLine, TEXT("-PARAMS="), Params);
+
 		if (Params.Len() > 0)
 		{
-			// remove from the command line
+			// remove from the command line & trim quotes
 			NewCommandLine = NewCommandLine.Replace(*Params, TEXT(""));
-
-			// trim the quotes
 			Params = Params.TrimQuotes();
 		}
 	}
@@ -581,25 +44,26 @@ int32 UnrealFrontendMain( const TCHAR* CommandLine )
 		NewCommandLine += TEXT(" -Messaging");
 	}
 
+	// initialize core
 	GEngineLoop.PreInit(*NewCommandLine);
-
-	// Tell the module manager is may now process newly-loaded UObjects when new C++ modules are loaded
 	FModuleManager::Get().StartProcessingNewlyLoadedObjects();
 
-	bool bReturnValue = true;
+	bool Succeeded = true;
+
+	// execute desired command
 	if (bRunCommand)
 	{
 		if (Command.Equals(TEXT("PACKAGE"), ESearchCase::IgnoreCase))
 		{
-			RunPackageCommand();
+			FPackageCommand::Run();
 		}
 		else if (Command.Equals(TEXT("DEPLOY"), ESearchCase::IgnoreCase))
 		{
-			bReturnValue = RunDeployCommand();
+			Succeeded = FDeployCommand::Run();
 		}
 		else if (Command.Equals(TEXT("LAUNCH"), ESearchCase::IgnoreCase))
 		{
-			bReturnValue = RunLaunchCommand(Params);
+			Succeeded = FLaunchCommand::Run(Params);
 		}
 		else if (Command.Equals(TEXT("CONVERT"), ESearchCase::IgnoreCase))
 		{
@@ -608,9 +72,10 @@ int32 UnrealFrontendMain( const TCHAR* CommandLine )
 	}
 	else
 	{
-		RunUI();
+		FUserInterfaceCommand::Run();
 	}
 
+	// shut down
 	FEngineLoop::AppPreExit();
 	FModuleManager::Get().UnloadModulesAtShutdown();
 
@@ -620,5 +85,5 @@ int32 UnrealFrontendMain( const TCHAR* CommandLine )
 
 	FTaskGraphInterface::Shutdown();
 
-	return bReturnValue ? 0 : -1;
+	return Succeeded ? 0 : -1;
 }
diff --git a/Engine/Source/Programs/UnrealFrontend/Private/UnrealFrontendMain.h b/Engine/Source/Programs/UnrealFrontend/Private/UnrealFrontendMain.h
--- a/Engine/Source/Programs/UnrealFrontend/Private/UnrealFrontendMain.h
+++ b/Engine/Source/Programs/UnrealFrontend/Private/UnrealFrontendMain.h
@@ -2,11 +2,6 @@
 
 #pragma once
 
-#include "Core.h"
-#include "Messaging.h"
-#include "Slate.h"
-#include "StandaloneRenderer.h"
-
 
 /**
  * The application's main function.
diff --git a/Engine/Source/Programs/UnrealFrontend/Private/UnrealFrontendPrivatePCH.h b/Engine/Source/Programs/UnrealFrontend/Private/UnrealFrontendPrivatePCH.h
new file mode 100644
--- /dev/null
+++ b/Engine/Source/Programs/UnrealFrontend/Private/UnrealFrontendPrivatePCH.h
@@ -0,0 +1,5 @@
+// Copyright 1998-2014 Epic Games, Inc. All Rights Reserved.
+
+#pragma once
+
+#include "Core.h"
diff --git a/Engine/Source/Programs/UnrealFrontend/Private/Windows/WindowsUnrealFrontendMain.cpp b/Engine/Source/Programs/UnrealFrontend/Private/Windows/WindowsUnrealFrontendMain.cpp
--- a/Engine/Source/Programs/UnrealFrontend/Private/Windows/WindowsUnrealFrontendMain.cpp
+++ b/Engine/Source/Programs/UnrealFrontend/Private/Windows/WindowsUnrealFrontendMain.cpp
@@ -1,5 +1,6 @@
 // Copyright 1998-2014 Epic Games, Inc. All Rights Reserved.
 
+#include "UnrealFrontendPrivatePCH.h"
 #include "UnrealFrontendMain.h"
 #include "ExceptionHandling.h"
 #include "LaunchEngineLoop.h"
diff --git a/Engine/Source/Programs/UnrealFrontend/UnrealFrontend.Build.cs b/Engine/Source/Programs/UnrealFrontend/UnrealFrontend.Build.cs
--- a/Engine/Source/Programs/UnrealFrontend/UnrealFrontend.Build.cs
+++ b/Engine/Source/Programs/UnrealFrontend/UnrealFrontend.Build.cs
@@ -11,6 +11,7 @@ public class UnrealFrontend : ModuleRules
 		PrivateIncludePaths.AddRange(
 			new string[] {
 				"Programs/UnrealFrontend/Private",
+				"Programs/UnrealFrontend/Private/Commands",
 				"Runtime/Launch/Private",					// for LaunchEngineLoop.cpp include
 			}
 		);
