* Fixed a problem with the settings for UFE which made it so it wouldn't properly add all of the dynamic modules.
UE-4880 - HTML5 not listed in UFE
#HTML5
#Android

diff --git a/Engine/Source/Programs/UnrealFrontend/UnrealFrontend.Target.cs b/Engine/Source/Programs/UnrealFrontend/UnrealFrontend.Target.cs
--- a/Engine/Source/Programs/UnrealFrontend/UnrealFrontend.Target.cs
+++ b/Engine/Source/Programs/UnrealFrontend/UnrealFrontend.Target.cs
@@ -43,7 +43,6 @@ public class UnrealFrontendTarget : TargetRules
 		UEBuildConfiguration.bBuildEditor = false;
 		UEBuildConfiguration.bCompileAgainstEngine = false;
 		UEBuildConfiguration.bCompileAgainstCoreUObject = true;
-		UEBuildConfiguration.bCompileLeanAndMeanUE = true;
 		UEBuildConfiguration.bForceBuildTargetPlatforms = true;
 		UEBuildConfiguration.bCompileWithStatsWithoutEngine = true;
 		UEBuildConfiguration.bCompileWithPluginSupport = true;
