Profiler - Stat memory dump command added to UFE, usage: -run=MEMORYDUMP -INFILE=<filename>

diff --git a/Engine/Source/Programs/UnrealFrontend/Private/Commands/StatsDumpMemoryCommand.cpp b/Engine/Source/Programs/UnrealFrontend/Private/Commands/StatsDumpMemoryCommand.cpp
new file mode 100644
--- /dev/null
+++ b/Engine/Source/Programs/UnrealFrontend/Private/Commands/StatsDumpMemoryCommand.cpp
@@ -0,0 +1,902 @@
+// Copyright 1998-2015 Epic Games, Inc. All Rights Reserved.
+
+#include "UnrealFrontendPrivatePCH.h"
+#include "StatsDumpMemoryCommand.h"
+#include "DiagnosticTable.h"
+
+static const double NumSecondsBetweenLogs = 5.0;
+
+void FStatsMemoryDumpCommand::InternalRun()
+{
+	FParse::Value( FCommandLine::Get(), TEXT( "-INFILE=" ), SourceFilepath );
+
+	const int64 Size = IFileManager::Get().FileSize( *SourceFilepath );
+	if( Size < 4 )
+	{
+		UE_LOG( LogStats, Error, TEXT( "Could not open: %s" ), *SourceFilepath );
+		return;
+	}
+	TAutoPtr<FArchive> FileReader( IFileManager::Get().CreateFileReader( *SourceFilepath ) );
+	if( !FileReader )
+	{
+		UE_LOG( LogStats, Error, TEXT( "Could not open: %s" ), *SourceFilepath );
+		return;
+	}
+
+	if( !Stream.ReadHeader( *FileReader ) )
+	{
+		UE_LOG( LogStats, Error, TEXT( "Could not open, bad magic: %s" ), *SourceFilepath );
+		return;
+	}
+
+	UE_LOG( LogStats, Warning, TEXT( "Reading a raw stats file for memory profiling: %s" ), *SourceFilepath );
+
+	const bool bIsFinalized = Stream.Header.IsFinalized();
+	check( bIsFinalized );
+	check( Stream.Header.Version == EStatMagicWithHeader::VERSION_5 );
+	StatsThreadStats.MarkAsLoaded();
+
+	TArray<FStatMessage> Messages;
+	if( Stream.Header.bRawStatsFile )
+	{
+		FScopeLogTime SLT( TEXT( "FStatsMemoryDumpCommand::InternalRun" ), nullptr, FScopeLogTime::ScopeLog_Seconds );
+
+		// Read metadata.
+		TArray<FStatMessage> MetadataMessages;
+		Stream.ReadFNamesAndMetadataMessages( *FileReader, MetadataMessages );
+		StatsThreadStats.ProcessMetaDataOnly( MetadataMessages );
+
+		// Find all UObject metadata messages.
+		for( const auto& Meta : MetadataMessages )
+		{
+			FName LongName = Meta.NameAndInfo.GetRawName();
+			const FString Desc = FStatNameAndInfo::GetDescriptionFrom( LongName );
+			const bool bContainsUObject = Desc.Contains( TEXT( "//" ) );
+			if( bContainsUObject )
+			{
+				UObjectNames.Add( LongName );
+			}
+		}
+
+		const int64 CurrentFilePos = FileReader->Tell();
+
+		// Update profiler's metadata.
+		CreateThreadsMapping();
+
+		// Read frames offsets.
+		Stream.ReadFramesOffsets( *FileReader );
+
+		// Buffer used to store the compressed and decompressed data.
+		TArray<uint8> SrcArray;
+		TArray<uint8> DestArray;
+		const bool bHasCompressedData = Stream.Header.HasCompressedData();
+		check( bHasCompressedData );
+
+		TMap<int64, FStatPacketArray> CombinedHistory;
+		int64 TotalDataSize = 0;
+		int64 TotalStatMessagesNum = 0;
+		int64 MaximumPacketSize = 0;
+		// Read all packets sequentially, force by the memory profiler which is now a part of the raw stats.
+		// !!CAUTION!! Frame number in the raw stats is pointless, because it is time based, not frame based.
+		// Background threads usually execute time consuming operations, so the frame number won't be valid.
+		// Needs to be combined by the thread and the time, not by the frame number.
+		{
+			// Display log information once per 5 seconds to avoid spamming.
+			double PreviousSeconds = FPlatformTime::Seconds();
+			const int64 FrameOffset0 = Stream.FramesInfo[0].FrameFileOffset;
+			FileReader->Seek( FrameOffset0 );
+
+			const int64 FileSize = FileReader->TotalSize();
+
+			while( FileReader->Tell() < FileSize )
+			{
+				// Read the compressed data.
+				FCompressedStatsData UncompressedData( SrcArray, DestArray );
+				*FileReader << UncompressedData;
+				if( UncompressedData.HasReachedEndOfCompressedData() )
+				{
+					break;
+				}
+
+				FMemoryReader MemoryReader( DestArray, true );
+
+				FStatPacket* StatPacket = new FStatPacket();
+				Stream.ReadStatPacket( MemoryReader, *StatPacket );
+
+				const int64 FrameNum = StatPacket->Frame;
+				FStatPacketArray& Frame = CombinedHistory.FindOrAdd( FrameNum );
+
+				// Check if we need to combine packets from the same thread.
+				FStatPacket** CombinedPacket = Frame.Packets.FindByPredicate( [&]( FStatPacket* Item ) -> bool
+				{
+					return Item->ThreadId == StatPacket->ThreadId;
+				} );
+
+				if( CombinedPacket )
+				{
+					(*CombinedPacket)->StatMessages += StatPacket->StatMessages;
+				}
+				else
+				{
+					Frame.Packets.Add( StatPacket );
+				}
+
+				const double CurrentSeconds = FPlatformTime::Seconds();
+				if( CurrentSeconds > PreviousSeconds + NumSecondsBetweenLogs )
+				{
+					const int32 PctPos = int32( 100.0*FileReader->Tell() / FileSize );
+					UE_LOG( LogStats, Log, TEXT( "%3i%% %10llu (%.1f MB) read messages, last read frame %4i" ), PctPos, TotalStatMessagesNum, TotalDataSize / 1024.0f / 1024.0f, StatPacket->Frame );
+					PreviousSeconds = CurrentSeconds;
+				}
+
+				const int64 PacketSize = StatPacket->StatMessages.GetAllocatedSize();
+				TotalDataSize += PacketSize;
+				MaximumPacketSize = FMath::Max( MaximumPacketSize, PacketSize );
+				TotalStatMessagesNum += StatPacket->StatMessages.Num();
+			}
+		}
+
+		UE_LOG( LogStats, Warning, TEXT( "TotalPacketSize: %.1f MB, Max: %1f MB" ),
+				TotalDataSize / 1024.0f / 1024.0f,
+				MaximumPacketSize / 1024.0f / 1024.0f );
+
+		TArray<int64> Frames;
+		CombinedHistory.GenerateKeyArray( Frames );
+		Frames.Sort();
+		const int64 MiddleFrame = Frames[Frames.Num() / 2];
+
+		ProcessMemoryOperations( CombinedHistory );
+	}
+}
+
+
+static const TCHAR* CallstackSeparator = TEXT( "+" );
+
+static FString EncodeCallstack( const TArray<FName>& Callstack )
+{
+	FString Result;
+	for( const auto& Name : Callstack )
+	{
+		Result += TTypeToString<int32>::ToString( (int32)Name.GetComparisonIndex() );
+		Result += CallstackSeparator;
+	}
+	return Result;
+}
+
+static void DecodeCallstack( const FName& EncodedCallstack, TArray<FString>& out_DecodedCallstack )
+{
+	EncodedCallstack.ToString().ParseIntoArray( &out_DecodedCallstack, CallstackSeparator, true );
+}
+
+/** Simple allocation info. Assumes the sequence tag are unique and doesn't wrap. */
+struct FAllocationInfo
+{
+	uint64 Ptr;
+	int64 Size;
+	FName EncodedCallstack;
+	uint32 SequenceTag;
+	EMemoryOperation Op;
+	bool bHasBrokenCallstack;
+
+	FAllocationInfo(
+		uint64 InPtr,
+		int64 InSize,
+		const TArray<FName>& InCallstack,
+		uint32 InSequenceTag,
+		EMemoryOperation InOp,
+		bool bInHasBrokenCallstack
+		)
+		: Ptr( InPtr )
+		, Size( InSize )
+		, EncodedCallstack( *EncodeCallstack( InCallstack ) )
+		, SequenceTag( InSequenceTag )
+		, Op( InOp )
+		, bHasBrokenCallstack( bInHasBrokenCallstack )
+	{}
+
+	FAllocationInfo( const FAllocationInfo& Other )
+		: Ptr( Other.Ptr )
+		, Size( Other.Size )
+		, EncodedCallstack( Other.EncodedCallstack )
+		, SequenceTag( Other.SequenceTag )
+		, Op( Other.Op )
+		, bHasBrokenCallstack( Other.bHasBrokenCallstack )
+	{}
+
+// 	FAllocationInfo()
+// 		: Ptr( 0 )
+// 		, Size( 0 )
+// 		, SequenceTag( 0 )
+// 		, Op( EMemoryOperation::Invalid )
+// 		, bHasBrokenCallstack( false )
+// 	{}
+
+
+	bool operator<(const FAllocationInfo& Other) const
+	{
+		return SequenceTag < Other.SequenceTag;
+	}
+
+	bool operator==(const FAllocationInfo& Other) const
+	{
+		return SequenceTag == Other.SequenceTag;
+	}
+};
+
+struct FAllocationInfoGreater
+{
+	FORCEINLINE bool operator()( const FAllocationInfo& A, const FAllocationInfo& B ) const
+	{
+		return B.Size < A.Size;
+	}
+};
+
+
+struct FSizeAndCount
+{
+	uint64 Size;
+	uint64 Count;
+	FSizeAndCount()
+		: Size( 0 )
+		, Count( 0 )
+	{}
+};
+
+struct FSizeAndCountGreater
+{
+	FORCEINLINE bool operator()( const FSizeAndCount& A, const FSizeAndCount& B ) const
+	{
+		return B.Size < A.Size;
+	}
+};
+
+/** Holds stats stack state, used to preserve continuity when the game frame has changed. */
+struct FStackState
+{
+	FStackState()
+		: bIsBrokenCallstack( false )
+	{}
+
+	/** Call stack. */
+	TArray<FName> Stack;
+
+	/** Current function name. */
+	FName Current;
+
+	/** Whether this callstack is marked as broken due to mismatched start and end scope cycles. */
+	bool bIsBrokenCallstack;
+};
+
+static FString GetCallstack( const FName& EncodedCallstack )
+{
+	FString Result;
+
+	TArray<FString> DecodedCallstack;
+	DecodeCallstack( EncodedCallstack, DecodedCallstack );
+
+	for( int32 Index = DecodedCallstack.Num() - 1; Index >= 0; --Index )
+	{
+		NAME_INDEX NameIndex = 0;
+		TTypeFromString<NAME_INDEX>::FromString( NameIndex, *DecodedCallstack[Index] );
+
+		const FName LongName = FName( NameIndex, NameIndex, 0 );
+
+		const FString ShortName = FStatNameAndInfo::GetShortNameFrom( LongName ).ToString();
+		//const FString Group = FStatNameAndInfo::GetGroupNameFrom( LongName ).ToString();
+		FString Desc = FStatNameAndInfo::GetDescriptionFrom( LongName );
+		Desc.Trim();
+
+		if( Desc.Len() == 0 )
+		{
+			Result += ShortName;
+		}
+		else
+		{
+			Result += Desc;
+		}
+
+		if( Index > 0 )
+		{
+			Result += TEXT( " <- " );
+		}
+	}
+
+	Result.ReplaceInline( TEXT( "STAT_" ), TEXT( "" ), ESearchCase::CaseSensitive );
+	return Result;
+}
+
+void FStatsMemoryDumpCommand::ProcessMemoryOperations( const TMap<int64, FStatPacketArray>& CombinedHistory )
+{
+	// This is only example code, no fully implemented, may sometimes crash.
+	// This code is not optimized. 
+	double PreviousSeconds = FPlatformTime::Seconds();
+	uint64 NumMemoryOperations = 0;
+
+	// Generate frames
+	TArray<int64> Frames;
+	CombinedHistory.GenerateKeyArray( Frames );
+	Frames.Sort();
+
+	// Raw stats callstack for this stat packet array.
+	TMap<FName, FStackState> StackStates;
+
+	// All allocation ordered by the sequence tag.
+	// There is an assumption that the sequence tag will not turn-around.
+	//TMap<uint32, FAllocationInfo> SequenceAllocationMap;
+	TArray<FAllocationInfo> SequenceAllocationArray;
+
+	// Pass 1.
+	// Read all stats messages, parse all memory operations and decode callstacks.
+	const int64 FirstFrame = 0;
+	PreviousSeconds -= NumSecondsBetweenLogs;
+	for( int32 FrameIndex = 0; FrameIndex < Frames.Num(); ++FrameIndex )
+	{
+		const double CurrentSeconds = FPlatformTime::Seconds();
+		if( CurrentSeconds > PreviousSeconds + NumSecondsBetweenLogs )
+		{
+			UE_LOG( LogStats, Warning, TEXT( "Processing frame %i/%i" ), FrameIndex+1, Frames.Num() );
+			PreviousSeconds = CurrentSeconds;
+		}
+
+		const int64 TargetFrame = Frames[FrameIndex];
+		const int64 Diff = TargetFrame - FirstFrame;
+		const FStatPacketArray& Frame = CombinedHistory.FindChecked( TargetFrame );
+
+		bool bAtLeastOnePacket = false;
+		for( int32 PacketIndex = 0; PacketIndex < Frame.Packets.Num(); PacketIndex++ )
+		{
+			const double CurrentSeconds = FPlatformTime::Seconds();
+			if( CurrentSeconds > PreviousSeconds + NumSecondsBetweenLogs )
+			{
+				UE_LOG( LogStats, Log, TEXT( "Processing packet %i/%i" ), PacketIndex, Frame.Packets.Num() );
+				PreviousSeconds = CurrentSeconds;
+				bAtLeastOnePacket = true;
+			}
+
+			const FStatPacket& StatPacket = *Frame.Packets[PacketIndex];
+			const FName& ThreadFName = StatsThreadStats.Threads.FindChecked( StatPacket.ThreadId );
+			const uint32 NewThreadID = ThreadIDtoStatID.FindChecked( StatPacket.ThreadId );
+
+			FStackState* StackState = StackStates.Find( ThreadFName );
+			if( !StackState )
+			{
+				StackState = &StackStates.Add( ThreadFName );
+				StackState->Stack.Add( ThreadFName );
+				StackState->Current = ThreadFName;
+			}
+
+			const FStatMessagesArray& Data = StatPacket.StatMessages;
+
+			int32 LastPct = 0;
+			const int32 NumDataElements = Data.Num();
+			const int32 OnerPercent = FMath::Max( NumDataElements / 100, 1024 );
+			bool bAtLeastOneMessage = false;
+			for( int32 Index = 0; Index < NumDataElements; Index++ )
+			{
+				if( Index % OnerPercent )
+				{
+					const double CurrentSeconds = FPlatformTime::Seconds();
+					if( CurrentSeconds > PreviousSeconds + NumSecondsBetweenLogs )
+					{
+						const int32 CurrentPct = int32( 100.0*(Index + 1) / NumDataElements );
+						UE_LOG( LogStats, Log, TEXT( "Processing %3i%% (%i/%i) stat messages" ), CurrentPct, Index, NumDataElements );
+						PreviousSeconds = CurrentSeconds;
+						bAtLeastOneMessage = true;
+					}
+				}
+
+				const FStatMessage& Item = Data[Index];
+
+				const EStatOperation::Type Op = Item.NameAndInfo.GetField<EStatOperation>();
+				const FName RawName = Item.NameAndInfo.GetRawName();
+
+				if( Op == EStatOperation::CycleScopeStart || Op == EStatOperation::CycleScopeEnd || Op == EStatOperation::Memory )
+				{
+					if( Op == EStatOperation::CycleScopeStart )
+					{
+						StackState->Stack.Add( RawName );
+						StackState->Current = RawName;
+					}
+					else if( Op == EStatOperation::Memory )
+					{
+						// Experimental code used only to test the implementation.
+						// First memory operation is Alloc or Free
+						const uint64 EncodedPtr = Item.GetValue_Ptr();
+						const bool bIsAlloc = (EncodedPtr & (uint64)EMemoryOperation::Alloc) != 0;
+						const bool bIsFree = (EncodedPtr & (uint64)EMemoryOperation::Free) != 0;
+						const uint64 Ptr = EncodedPtr & ~(uint64)EMemoryOperation::Mask;
+						if( bIsAlloc )
+						{
+							NumMemoryOperations++;
+							// @see FStatsMallocProfilerProxy::TrackAlloc
+							// After alloc ptr message there is always alloc size message and the sequence tag.
+							Index++;
+							const FStatMessage& AllocSizeMessage = Data[Index];
+							const int64 AllocSize = AllocSizeMessage.GetValue_int64();
+
+							// Read operation sequence tag.
+							Index++;
+							const FStatMessage& SequenceTagMessage = Data[Index];
+							const uint32 SequenceTag = SequenceTagMessage.GetValue_int64();
+
+							// Create a callstack.
+							TArray<FName> StatsBasedCallstack;
+							for( const auto& StackName : StackState->Stack )
+							{
+								StatsBasedCallstack.Add( StackName );
+							}
+
+							// Add a new allocation.
+							SequenceAllocationArray.Add(
+								FAllocationInfo(
+								Ptr,
+								AllocSize,
+								StatsBasedCallstack,
+								SequenceTag,
+								EMemoryOperation::Alloc,
+								StackState->bIsBrokenCallstack
+								) );
+						}
+						else if( bIsFree )
+						{
+							NumMemoryOperations++;
+							// Read operation sequence tag.
+							Index++;
+							const FStatMessage& SequenceTagMessage = Data[Index];
+							const uint32 SequenceTag = SequenceTagMessage.GetValue_int64();
+
+							// Create a callstack.
+							/*
+							TArray<FName> StatsBasedCallstack;
+							for( const auto& RawName : StackState->Stack )
+							{
+								StatsBasedCallstack.Add( RawName );
+							}
+							*/
+
+							// Add a new free.
+							SequenceAllocationArray.Add(
+								FAllocationInfo(
+								Ptr,		
+								0,
+								TArray<FName>()/*StatsBasedCallstack*/,					
+								SequenceTag,
+								EMemoryOperation::Free,
+								StackState->bIsBrokenCallstack
+								) );
+						}
+						else
+						{
+							UE_LOG( LogStats, Warning, TEXT( "Pointer from a memory operation is invalid" ) );
+						}
+					}
+					else if( Op == EStatOperation::CycleScopeEnd )
+					{
+						if( StackState->Stack.Num() > 1 )
+						{
+							const FName ScopeStart = StackState->Stack.Pop();
+							const FName ScopeEnd = Item.NameAndInfo.GetRawName();
+
+							check( ScopeStart == ScopeEnd );
+
+							StackState->Current = StackState->Stack.Last();
+
+							// The stack should be ok, but it may be partially broken.
+							// This will happen if memory profiling starts in the middle of executing a background thread.
+							StackState->bIsBrokenCallstack = false;
+						}
+						else
+						{
+							const FName ShortName = Item.NameAndInfo.GetShortName();
+
+							UE_LOG( LogStats, Warning, TEXT( "Broken cycle scope end %s/%s, current %s" ),
+									*ThreadFName.ToString(),
+									*ShortName.ToString(),
+									*StackState->Current.ToString() );
+
+							// The stack is completely broken, only has the thread name and the last cycle scope.
+							// Rollback to the thread node.
+							StackState->bIsBrokenCallstack = true;
+							StackState->Stack.Empty();
+							StackState->Stack.Add( ThreadFName );
+							StackState->Current = ThreadFName;
+						}
+					}
+				}
+			}
+			if( bAtLeastOneMessage )
+			{
+				PreviousSeconds -= NumSecondsBetweenLogs;
+			}
+		}
+		if( bAtLeastOnePacket )
+		{
+			PreviousSeconds -= NumSecondsBetweenLogs;
+		}
+	}
+
+	UE_LOG( LogStats, Warning, TEXT( "NumMemoryOperations:   %llu" ), NumMemoryOperations );
+	UE_LOG( LogStats, Warning, TEXT( "SequenceAllocationNum: %i" ), SequenceAllocationArray.Num() );
+
+	// Pass 2.
+	/*
+	TMap<uint32,FAllocationInfo> UniqueSeq;
+	TMultiMap<uint32,FAllocationInfo> OriginalAllocs;
+	TMultiMap<uint32,FAllocationInfo> BrokenAllocs;
+	for( const FAllocationInfo& Alloc : SequenceAllocationArray )
+	{
+		const FAllocationInfo* Found = UniqueSeq.Find(Alloc.SequenceTag);
+		if( !Found )
+		{
+			UniqueSeq.Add(Alloc.SequenceTag,Alloc);
+		}
+		else
+		{
+			OriginalAllocs.Add(Alloc.SequenceTag, *Found);
+			BrokenAllocs.Add(Alloc.SequenceTag, Alloc);
+		}
+	}
+	*/
+
+	// Sort all memory operation by the sequence tag, iterate through all operation and generate memory usage.
+	SequenceAllocationArray.Sort( TLess<FAllocationInfo>() );
+
+	// Alive allocations.
+	TMap<uint64, FAllocationInfo> AllocationMap;
+	TMultiMap<uint64, FAllocationInfo> FreeWithoutAllocMap;
+	TMultiMap<uint64, FAllocationInfo> DuplicatedAllocMap;
+	int32 NumDuplicatedMemoryOperations = 0;
+	int32 NumFWAMemoryOperations = 0; // FreeWithoutAlloc
+
+	UE_LOG( LogStats, Warning, TEXT( "Generating memory operations map" ) );
+	const int32 NumSequenceAllocations = SequenceAllocationArray.Num();
+	const int32 OnePercent = FMath::Max( NumSequenceAllocations / 100, 1024 );
+	for( int32 Index = 0; Index < NumSequenceAllocations; Index++ )
+	{
+		if( Index % OnePercent )
+		{
+			const double CurrentSeconds = FPlatformTime::Seconds();
+			if( CurrentSeconds > PreviousSeconds + NumSecondsBetweenLogs )
+			{
+				const int32 CurrentPct = int32( 100.0*(Index + 1) / NumSequenceAllocations );
+				UE_LOG( LogStats, Log, TEXT( "Processing allocations %3i%% (%10i/%10i)" ), CurrentPct, Index + 1, NumSequenceAllocations );
+				PreviousSeconds = CurrentSeconds;
+			}
+		}
+
+		const FAllocationInfo& Alloc = SequenceAllocationArray[Index];
+		const EMemoryOperation MemOp = Alloc.Op;
+		const uint64 Ptr = Alloc.Ptr;
+		const int64 Size = Alloc.Size;
+		const uint32 SequenceTag = Alloc.SequenceTag;
+
+		if( MemOp == EMemoryOperation::Alloc )
+		{
+			const FAllocationInfo* Found = AllocationMap.Find( Ptr );
+
+			if( !Found )
+			{
+				AllocationMap.Add( Ptr, Alloc );
+			}
+			else
+			{
+				const FAllocationInfo* FoundAndFreed = FreeWithoutAllocMap.Find( Found->Ptr );
+				const FAllocationInfo* FoundAndAllocated = FreeWithoutAllocMap.Find( Alloc.Ptr );
+
+#if	_DEBUG
+				if( FoundAndFreed )
+				{
+					const FString FoundAndFreedCallstack = GetCallstack( FoundAndFreed->EncodedCallstack );
+				}
+
+				if( FoundAndAllocated )
+				{
+					const FString FoundAndAllocatedCallstack = GetCallstack( FoundAndAllocated->EncodedCallstack );
+				}
+
+				NumDuplicatedMemoryOperations++;
+
+
+				const FString FoundCallstack = GetCallstack( Found->EncodedCallstack );
+				const FString AllocCallstack = GetCallstack( Alloc.EncodedCallstack );
+#endif // _DEBUG
+
+				// Replace pointer.
+				AllocationMap.Add( Ptr, Alloc );
+				// Store the old pointer.
+				DuplicatedAllocMap.Add( Ptr, *Found );
+			}
+		}
+		else if( MemOp == EMemoryOperation::Free )
+		{
+			const FAllocationInfo* Found = AllocationMap.Find( Ptr );
+			if( Found )
+			{
+				const bool bIsValid = Alloc.SequenceTag > Found->SequenceTag;
+				if( !bIsValid )
+				{
+					UE_LOG( LogStats, Warning, TEXT( "InvalidFree Ptr: %llu, Seq: %i/%i" ), Ptr, SequenceTag, Found->SequenceTag );
+				}
+				AllocationMap.Remove( Ptr );
+			}
+			else
+			{
+				FreeWithoutAllocMap.Add( Ptr, Alloc );
+				NumFWAMemoryOperations++;
+			}
+		}
+	}
+
+	UE_LOG( LogStats, Warning, TEXT( "NumDuplicatedMemoryOperations: %i" ), NumDuplicatedMemoryOperations );
+	UE_LOG( LogStats, Warning, TEXT( "NumFWAMemoryOperations:        %i" ), NumFWAMemoryOperations );
+
+	// Dump problematic allocations
+	DuplicatedAllocMap.ValueSort( FAllocationInfoGreater() );
+	//FreeWithoutAllocMap
+
+	uint64 TotalDuplicatedMemory = 0;
+	for( const auto& It : DuplicatedAllocMap )
+	{
+		const FAllocationInfo& Alloc = It.Value;
+		TotalDuplicatedMemory += Alloc.Size;
+	}
+
+	UE_LOG( LogStats, Warning, TEXT( "Dumping duplicated alloc map" ) );
+	const float MaxPctDisplayed = 0.80f;
+	uint64 DisplayedSoFar = 0;
+	for( const auto& It : DuplicatedAllocMap )
+	{
+		const FAllocationInfo& Alloc = It.Value;
+		const FString AllocCallstack = GetCallstack( Alloc.EncodedCallstack );
+		UE_LOG( LogStats, Log, TEXT( "%lli (%.2f MB) %s" ), Alloc.Size, Alloc.Size / 1024.0f / 1024.0f, *AllocCallstack );
+
+		DisplayedSoFar += Alloc.Size;
+
+		const float CurrentPct = (float)DisplayedSoFar / (float)TotalDuplicatedMemory;
+		if( CurrentPct > MaxPctDisplayed )
+		{
+			break;
+		}
+	}
+
+	GenerateMemoryUsageReport( AllocationMap );
+}
+
+void FStatsMemoryDumpCommand::GenerateMemoryUsageReport( const TMap<uint64, FAllocationInfo>& AllocationMap )
+{
+	if( AllocationMap.Num() == 0 )
+	{
+		UE_LOG( LogStats, Warning, TEXT( "There are no allocations, make sure memory profiler is enabled" ) );
+	}
+	else
+	{
+		ProcessingUObjectAllocations( AllocationMap );
+		ProcessingScopedAllocations( AllocationMap );
+	}
+}
+
+void FStatsMemoryDumpCommand::CreateThreadsMapping()
+{
+	int32 UniqueID = 15;
+
+	for( auto It = StatsThreadStats.ShortNameToLongName.CreateConstIterator(); It; ++It )
+	{
+		const FStatMessage& LongName = It.Value();
+
+		const FName GroupName = LongName.NameAndInfo.GetGroupName();
+		
+		// Setup thread id to stat id.
+		if( GroupName == FStatConstants::NAME_ThreadGroup )
+		{
+			const FName StatName = It.Key();
+			UniqueID++;
+
+			uint32 ThreadID = 0;
+			for( auto ThreadsIt = StatsThreadStats.Threads.CreateConstIterator(); ThreadsIt; ++ThreadsIt )
+			{
+				if( ThreadsIt.Value() == StatName )
+				{
+					ThreadID = ThreadsIt.Key();
+					break;
+				}
+			}
+			ThreadIDtoStatID.Add( ThreadID, UniqueID );
+
+			// Game thread is always NAME_GameThread
+			if( StatName == NAME_GameThread )
+			{
+				GameThreadID = ThreadID;
+			}
+		}
+	}
+}
+
+void FStatsMemoryDumpCommand::ProcessingScopedAllocations( const TMap<uint64, FAllocationInfo>& AllocationMap )
+{
+	// This code is not optimized. 
+	FScopeLogTime SLT( TEXT( "ProcessingScopedAllocations" ), nullptr, FScopeLogTime::ScopeLog_Seconds );
+
+	UE_LOG( LogStats, Warning, TEXT( "Processing scoped allocations" ) );
+
+	FDiagnosticTableViewer MemoryReport( *FDiagnosticTableViewer::GetUniqueTemporaryFilePath( TEXT( "MemoryReport-Scoped" ) ) );
+	// Write a row of headings for the table's columns.
+	MemoryReport.AddColumn( TEXT( "Size (bytes)" ) );
+	MemoryReport.AddColumn( TEXT( "Size (MB)" ) );
+	MemoryReport.AddColumn( TEXT( "Count" ) );
+	MemoryReport.AddColumn( TEXT( "Callstack" ) );
+	MemoryReport.CycleRow();
+
+
+	TMap<FName, FSizeAndCount> ScopedAllocations;
+
+	uint64 NumAllocations = 0;
+	uint64 TotalAllocatedMemory = 0;
+	for( const auto& It : AllocationMap )
+	{
+		const FAllocationInfo& Alloc = It.Value;
+		FSizeAndCount& SizeAndCount = ScopedAllocations.FindOrAdd( Alloc.EncodedCallstack );
+		SizeAndCount.Size += Alloc.Size;
+		SizeAndCount.Count += 1;
+
+		TotalAllocatedMemory += Alloc.Size;
+		NumAllocations++;
+	}
+
+	// Dump memory to the log.
+	ScopedAllocations.ValueSort( FSizeAndCountGreater() );
+
+	const float MaxPctDisplayed = 0.90f;
+	int32 CurrentIndex = 0;
+	uint64 DisplayedSoFar = 0;
+	UE_LOG( LogStats, Warning, TEXT( "Index, Size (Size MB), Count, Stat desc" ) );
+	for( const auto& It : ScopedAllocations )
+	{
+		const FSizeAndCount& SizeAndCount = It.Value;
+		const FName& EncodedCallstack = It.Key;
+
+		const FString AllocCallstack = GetCallstack( EncodedCallstack );
+
+		UE_LOG( LogStats, Log, TEXT( "%2i, %llu (%.2f MB), %llu, %s" ),
+				CurrentIndex,
+				SizeAndCount.Size,
+				SizeAndCount.Size / 1024.0f / 1024.0f,
+				SizeAndCount.Count,
+				*AllocCallstack );
+
+		// Dump stats
+		MemoryReport.AddColumn( TEXT( "%llu" ), SizeAndCount.Size );
+		MemoryReport.AddColumn( TEXT( "%.2f MB" ), SizeAndCount.Size / 1024.0f / 1024.0f );
+		MemoryReport.AddColumn( TEXT( "%llu" ), SizeAndCount.Count );
+		MemoryReport.AddColumn( *AllocCallstack );
+		MemoryReport.CycleRow();
+
+		CurrentIndex++;
+		DisplayedSoFar += SizeAndCount.Size;
+
+		const float CurrentPct = (float)DisplayedSoFar / (float)TotalAllocatedMemory;
+		if( CurrentPct > MaxPctDisplayed )
+		{
+			break;
+		}
+	}
+
+	UE_LOG( LogStats, Warning, TEXT( "Allocated memory: %llu bytes (%.2f MB)" ), TotalAllocatedMemory, TotalAllocatedMemory / 1024.0f / 1024.0f );
+
+	// Add a total row.
+	MemoryReport.CycleRow();
+	MemoryReport.CycleRow();
+	MemoryReport.CycleRow();
+	MemoryReport.AddColumn( TEXT( "%llu" ), TotalAllocatedMemory );
+	MemoryReport.AddColumn( TEXT( "%.2f MB" ), TotalAllocatedMemory / 1024.0f / 1024.0f );
+	MemoryReport.AddColumn( TEXT( "%llu" ), NumAllocations );
+	MemoryReport.AddColumn( TEXT( "TOTAL" ) );
+	MemoryReport.CycleRow();
+}
+
+void FStatsMemoryDumpCommand::ProcessingUObjectAllocations( const TMap<uint64, FAllocationInfo>& AllocationMap )
+{
+	// This code is not optimized. 
+	FScopeLogTime SLT( TEXT( "ProcessingUObjectAllocations" ), nullptr, FScopeLogTime::ScopeLog_Seconds );
+	UE_LOG( LogStats, Warning, TEXT( "Processing UObject allocations" ) );
+
+	FDiagnosticTableViewer MemoryReport( *FDiagnosticTableViewer::GetUniqueTemporaryFilePath( TEXT( "MemoryReport-UObject" ) ) );
+	// Write a row of headings for the table's columns.
+	MemoryReport.AddColumn( TEXT( "Size (bytes)" ) );
+	MemoryReport.AddColumn( TEXT( "Size (MB)" ) );
+	MemoryReport.AddColumn( TEXT( "Count" ) );
+	MemoryReport.AddColumn( TEXT( "UObject class" ) );
+	MemoryReport.CycleRow();
+
+	TMap<FName, FSizeAndCount> UObjectAllocations;
+
+	// To minimize number of calls to expensive DecodeCallstack.
+	TMap<FName,FName> UObjectCallstackToClassMapping;
+
+	uint64 NumAllocations = 0;
+	uint64 TotalAllocatedMemory = 0;
+	for( const auto& It : AllocationMap )
+	{
+		const FAllocationInfo& Alloc = It.Value;
+
+		FName UObjectClass = UObjectCallstackToClassMapping.FindRef( Alloc.EncodedCallstack );
+		if( UObjectClass == NAME_None )
+		{
+			TArray<FString> DecodedCallstack;
+			DecodeCallstack( Alloc.EncodedCallstack, DecodedCallstack );
+
+			for( int32 Index = DecodedCallstack.Num() - 1; Index >= 0; --Index )
+			{
+				NAME_INDEX NameIndex = 0;
+				TTypeFromString<NAME_INDEX>::FromString( NameIndex, *DecodedCallstack[Index] );
+
+				const FName LongName = FName( NameIndex, NameIndex, 0 );
+				const bool bValid = UObjectNames.Contains( LongName );
+				if( bValid )
+				{
+					const FString ObjectName = FStatNameAndInfo::GetDescriptionFrom( LongName );
+					UObjectClass = *ObjectName.Left( ObjectName.Find( TEXT( "//" ) ) );;
+					UObjectCallstackToClassMapping.Add( Alloc.EncodedCallstack, UObjectClass );
+					break;
+				}
+			}
+		}
+		
+		if( UObjectClass != NAME_None )
+		{
+			FSizeAndCount& SizeAndCount = UObjectAllocations.FindOrAdd( UObjectClass );
+			SizeAndCount.Size += Alloc.Size;
+			SizeAndCount.Count += 1;
+
+			TotalAllocatedMemory += Alloc.Size;
+			NumAllocations++;
+		}
+	}
+
+	// Dump memory to the log.
+	UObjectAllocations.ValueSort( FSizeAndCountGreater() );
+
+	const float MaxPctDisplayed = 0.90f;
+	int32 CurrentIndex = 0;
+	uint64 DisplayedSoFar = 0;
+	UE_LOG( LogStats, Warning, TEXT( "Index, Size (Size MB), Count, UObject class" ) );
+	for( const auto& It : UObjectAllocations )
+	{
+		const FSizeAndCount& SizeAndCount = It.Value;
+		const FName& UObjectClass = It.Key;
+
+		UE_LOG( LogStats, Log, TEXT( "%2i, %llu (%.2f MB), %llu, %s" ),
+				CurrentIndex,
+				SizeAndCount.Size,
+				SizeAndCount.Size / 1024.0f / 1024.0f,
+				SizeAndCount.Count,
+				*UObjectClass.GetPlainNameString() );
+
+		// Dump stats
+		MemoryReport.AddColumn( TEXT( "%llu" ), SizeAndCount.Size );
+		MemoryReport.AddColumn( TEXT( "%.2f MB" ), SizeAndCount.Size / 1024.0f / 1024.0f );
+		MemoryReport.AddColumn( TEXT( "%llu" ), SizeAndCount.Count );
+		MemoryReport.AddColumn( *UObjectClass.GetPlainNameString() );
+		MemoryReport.CycleRow();
+
+		CurrentIndex++;
+		DisplayedSoFar += SizeAndCount.Size;
+
+		const float CurrentPct = (float)DisplayedSoFar / (float)TotalAllocatedMemory;
+		if( CurrentPct > MaxPctDisplayed )
+		{
+			break;
+		}
+	}
+
+	UE_LOG( LogStats, Warning, TEXT( "Allocated memory: %llu bytes (%.2f MB)" ), TotalAllocatedMemory, TotalAllocatedMemory / 1024.0f / 1024.0f );
+
+	// Add a total row.
+	MemoryReport.CycleRow();
+	MemoryReport.CycleRow();
+	MemoryReport.CycleRow();
+	MemoryReport.AddColumn( TEXT( "%llu" ), TotalAllocatedMemory );
+	MemoryReport.AddColumn( TEXT( "%.2f MB" ), TotalAllocatedMemory / 1024.0f / 1024.0f );
+	MemoryReport.AddColumn( TEXT( "%llu" ), NumAllocations );
+	MemoryReport.AddColumn( TEXT( "TOTAL" ) );
+	MemoryReport.CycleRow();
+}
diff --git a/Engine/Source/Programs/UnrealFrontend/Private/Commands/StatsDumpMemoryCommand.h b/Engine/Source/Programs/UnrealFrontend/Private/Commands/StatsDumpMemoryCommand.h
new file mode 100644
--- /dev/null
+++ b/Engine/Source/Programs/UnrealFrontend/Private/Commands/StatsDumpMemoryCommand.h
@@ -0,0 +1,59 @@
+// Copyright 1998-2015 Epic Games, Inc. All Rights Reserved.
+
+#pragma once
+
+#include "StatsData.h"
+#include "StatsFile.h"
+
+struct FAllocationInfo;
+
+/** An example command loading a raw stats file and dumping memory usage. */
+class FStatsMemoryDumpCommand
+{
+public:
+
+	/** Executes the command. */
+	static void Run()
+	{
+		FStatsMemoryDumpCommand Instance;
+		Instance.InternalRun();
+	}
+
+protected:
+	/** Executes the command. */
+	void InternalRun();
+
+	/** Creates thread mapping, some kind of metadata. */
+	void CreateThreadsMapping();
+
+	/** Basic memory profiling, only for debugging purpose. */
+	void ProcessMemoryOperations( const TMap<int64, FStatPacketArray>& CombinedHistory );
+
+	/** Generates a basic memory usage report and prints it to the log. */
+	void GenerateMemoryUsageReport( const TMap<uint64, FAllocationInfo>& AllocationMap );
+
+	/** Generates scoped allocation statistics. */
+	void ProcessingScopedAllocations( const TMap<uint64, FAllocationInfo>& AllocationMap );
+
+	/** Generates UObject allocation statistics. */
+	void ProcessingUObjectAllocations( const TMap<uint64, FAllocationInfo>& AllocationMap );
+
+	/** Stats thread state, mostly used to manage the stats metadata. */
+	FStatsThreadState StatsThreadStats;
+	FStatsReadStream Stream;
+
+	/** All names that containa a path to an UObject. */
+	TSet<FName> UObjectNames;
+
+	/** Filepath to the raw stats file. */
+	FString SourceFilepath;
+
+	/**
+	 *	Helper map, as a part of migration into stats2, stored as ThreadID -> StatID.
+	 *	Used during creating a thread sample, so we can use the real stat id and get data graph for this thread
+	 */
+	TMap<uint32, uint32> ThreadIDtoStatID;
+
+	/** Game thread id. */
+	uint32 GameThreadID;
+};
diff --git a/Engine/Source/Programs/UnrealFrontend/Private/UnrealFrontendMain.cpp b/Engine/Source/Programs/UnrealFrontend/Private/UnrealFrontendMain.cpp
--- a/Engine/Source/Programs/UnrealFrontend/Private/UnrealFrontendMain.cpp
+++ b/Engine/Source/Programs/UnrealFrontend/Private/UnrealFrontendMain.cpp
@@ -8,6 +8,7 @@
 #include "LaunchCommand.h"
 #include "PackageCommand.h"
 #include "StatsConvertCommand.h"
+#include "StatsDumpMemoryCommand.h"
 #include "UserInterfaceCommand.h"
 
 
@@ -69,6 +70,10 @@ int32 UnrealFrontendMain( const TCHAR* CommandLine )
 		{
 			FStatsConvertCommand::Run();
 		}
+		else if( Command.Equals( TEXT("MEMORYDUMP"), ESearchCase::IgnoreCase ) )
+		{
+			FStatsMemoryDumpCommand::Run();
+		}
 	}
 	else
 	{
