Profiler - Fixed memory leak in the stat dump memory command

diff --git a/Engine/Source/Programs/UnrealFrontend/Private/Commands/StatsDumpMemoryCommand.cpp b/Engine/Source/Programs/UnrealFrontend/Private/Commands/StatsDumpMemoryCommand.cpp
--- a/Engine/Source/Programs/UnrealFrontend/Private/Commands/StatsDumpMemoryCommand.cpp
+++ b/Engine/Source/Programs/UnrealFrontend/Private/Commands/StatsDumpMemoryCommand.cpp
@@ -76,6 +76,7 @@ void FStatsMemoryDumpCommand::InternalRun()
 		int64 TotalDataSize = 0;
 		int64 TotalStatMessagesNum = 0;
 		int64 MaximumPacketSize = 0;
+		int64 TotalPacketsNum = 0;
 		// Read all packets sequentially, force by the memory profiler which is now a part of the raw stats.
 		// !!CAUTION!! Frame number in the raw stats is pointless, because it is time based, not frame based.
 		// Background threads usually execute time consuming operations, so the frame number won't be valid.
@@ -103,8 +104,8 @@ void FStatsMemoryDumpCommand::InternalRun()
 				FStatPacket* StatPacket = new FStatPacket();
 				Stream.ReadStatPacket( MemoryReader, *StatPacket );
 
-				const int64 FrameNum = StatPacket->Frame;
-				FStatPacketArray& Frame = CombinedHistory.FindOrAdd( FrameNum );
+				const int64 StatPacketFrameNum = StatPacket->Frame;
+				FStatPacketArray& Frame = CombinedHistory.FindOrAdd( StatPacketFrameNum );
 
 				// Check if we need to combine packets from the same thread.
 				FStatPacket** CombinedPacket = Frame.Packets.FindByPredicate( [&]( FStatPacket* Item ) -> bool
@@ -112,28 +113,54 @@ void FStatsMemoryDumpCommand::InternalRun()
 					return Item->ThreadId == StatPacket->ThreadId;
 				} );
 
+				const int64 PacketSize = StatPacket->StatMessages.GetAllocatedSize();
+				TotalStatMessagesNum += StatPacket->StatMessages.Num();
+
 				if( CombinedPacket )
 				{
+					TotalDataSize -= (*CombinedPacket)->StatMessages.GetAllocatedSize();
 					(*CombinedPacket)->StatMessages += StatPacket->StatMessages;
+					TotalDataSize += (*CombinedPacket)->StatMessages.GetAllocatedSize();
+
+					delete StatPacket;
 				}
 				else
 				{
 					Frame.Packets.Add( StatPacket );
+					TotalDataSize += PacketSize;
 				}
 
 				const double CurrentSeconds = FPlatformTime::Seconds();
 				if( CurrentSeconds > PreviousSeconds + NumSecondsBetweenLogs )
 				{
 					const int32 PctPos = int32( 100.0*FileReader->Tell() / FileSize );
-					UE_LOG( LogStats, Log, TEXT( "%3i%% %10llu (%.1f MB) read messages, last read frame %4i" ), PctPos, TotalStatMessagesNum, TotalDataSize / 1024.0f / 1024.0f, StatPacket->Frame );
+					UE_LOG( LogStats, Log, TEXT( "%3i%% %10llu (%.1f MB) read messages, last read frame %4i" ), PctPos, TotalStatMessagesNum, TotalDataSize / 1024.0f / 1024.0f, StatPacketFrameNum );
 					PreviousSeconds = CurrentSeconds;
 				}
+			
+				MaximumPacketSize = FMath::Max( MaximumPacketSize, PacketSize );			
+				TotalPacketsNum++;
+			}
+		}
 
-				const int64 PacketSize = StatPacket->StatMessages.GetAllocatedSize();
-				TotalDataSize += PacketSize;
-				MaximumPacketSize = FMath::Max( MaximumPacketSize, PacketSize );
-				TotalStatMessagesNum += StatPacket->StatMessages.Num();
+		// Dump frame stats
+		for( const auto& It : CombinedHistory )
+		{
+			const int64 FrameNum = It.Key;
+			int64 FramePacketsSize = 0;
+			int64 FrameStatMessages = 0;
+			int64 FramePackets = It.Value.Packets.Num(); // Threads
+			for( const auto& It2 : It.Value.Packets )
+			{
+				FramePacketsSize += It2->StatMessages.GetAllocatedSize();
+				FrameStatMessages += It2->StatMessages.Num();
 			}
+
+			UE_LOG( LogStats, Warning, TEXT( "Frame: %10llu/%3lli Size: %.1f MB / %10lli" ), 
+					FrameNum, 
+					FramePackets, 
+					FramePacketsSize / 1024.0f / 1024.0f,
+					FrameStatMessages );
 		}
 
 		UE_LOG( LogStats, Warning, TEXT( "TotalPacketSize: %.1f MB, Max: %1f MB" ),
