Stats - Fixed not working UObject stats in the memory dump command

diff --git a/Engine/Source/Programs/UnrealFrontend/Private/Commands/StatsDumpMemoryCommand.cpp b/Engine/Source/Programs/UnrealFrontend/Private/Commands/StatsDumpMemoryCommand.cpp
--- a/Engine/Source/Programs/UnrealFrontend/Private/Commands/StatsDumpMemoryCommand.cpp
+++ b/Engine/Source/Programs/UnrealFrontend/Private/Commands/StatsDumpMemoryCommand.cpp
@@ -50,7 +50,7 @@ void FStatsMemoryDumpCommand::InternalRun()
 		for( const auto& Meta : MetadataMessages )
 		{
 			FName LongName = Meta.NameAndInfo.GetRawName();
-			const FString Desc = FStatNameAndInfo::GetDescriptionFrom( LongName );
+			const FString Desc = FStatNameAndInfo::GetShortNameFrom( LongName ).GetPlainNameString();
 			const bool bContainsUObject = Desc.Contains( TEXT( "//" ) );
 			if( bContainsUObject )
 			{
