Stats - Fixed not working UObject stats in the memory dump command

diff --git a/Engine/Source/Programs/UnrealFrontend/Private/Commands/StatsDumpMemoryCommand.cpp b/Engine/Source/Programs/UnrealFrontend/Private/Commands/StatsDumpMemoryCommand.cpp
--- a/Engine/Source/Programs/UnrealFrontend/Private/Commands/StatsDumpMemoryCommand.cpp
+++ b/Engine/Source/Programs/UnrealFrontend/Private/Commands/StatsDumpMemoryCommand.cpp
@@ -860,7 +860,7 @@ void FStatsMemoryDumpCommand::ProcessingUObjectAllocations( const TMap<uint64, F
 				const bool bValid = UObjectNames.Contains( LongName );
 				if( bValid )
 				{
-					const FString ObjectName = FStatNameAndInfo::GetDescriptionFrom( LongName );
+					const FString ObjectName = FStatNameAndInfo::GetShortNameFrom( LongName ).GetPlainNameString();
 					UObjectClass = *ObjectName.Left( ObjectName.Find( TEXT( "//" ) ) );;
 					UObjectCallstackToClassMapping.Add( Alloc.EncodedCallstack, UObjectClass );
 					break;
