Fix UnrealFrontend shadow variables

diff --git a/Engine/Source/Programs/UnrealFrontend/Private/Commands/StatsDumpMemoryCommand.cpp b/Engine/Source/Programs/UnrealFrontend/Private/Commands/StatsDumpMemoryCommand.cpp
--- a/Engine/Source/Programs/UnrealFrontend/Private/Commands/StatsDumpMemoryCommand.cpp
+++ b/Engine/Source/Programs/UnrealFrontend/Private/Commands/StatsDumpMemoryCommand.cpp
@@ -358,12 +358,14 @@ void FStatsMemoryDumpCommand::ProcessMemoryOperations( const TMap<int64, FStatPa
 	PreviousSeconds -= NumSecondsBetweenLogs;
 	for( int32 FrameIndex = 0; FrameIndex < Frames.Num(); ++FrameIndex )
 	{
-		const double CurrentSeconds = FPlatformTime::Seconds();
-		if( CurrentSeconds > PreviousSeconds + NumSecondsBetweenLogs )
-		{
-			UE_LOG( LogStats, Warning, TEXT( "Processing frame %i/%i" ), FrameIndex+1, Frames.Num() );
-			PreviousSeconds = CurrentSeconds;
-		}
+        {
+            const double CurrentSeconds = FPlatformTime::Seconds();
+            if( CurrentSeconds > PreviousSeconds + NumSecondsBetweenLogs )
+            {
+                UE_LOG( LogStats, Warning, TEXT( "Processing frame %i/%i" ), FrameIndex+1, Frames.Num() );
+                PreviousSeconds = CurrentSeconds;
+            }
+        }
 
 		const int64 TargetFrame = Frames[FrameIndex];
 		const int64 Diff = TargetFrame - FirstFrame;
@@ -372,13 +374,15 @@ void FStatsMemoryDumpCommand::ProcessMemoryOperations( const TMap<int64, FStatPa
 		bool bAtLeastOnePacket = false;
 		for( int32 PacketIndex = 0; PacketIndex < Frame.Packets.Num(); PacketIndex++ )
 		{
-			const double CurrentSeconds = FPlatformTime::Seconds();
-			if( CurrentSeconds > PreviousSeconds + NumSecondsBetweenLogs )
-			{
-				UE_LOG( LogStats, Log, TEXT( "Processing packet %i/%i" ), PacketIndex, Frame.Packets.Num() );
-				PreviousSeconds = CurrentSeconds;
-				bAtLeastOnePacket = true;
-			}
+            {
+                const double CurrentSeconds = FPlatformTime::Seconds();
+                if( CurrentSeconds > PreviousSeconds + NumSecondsBetweenLogs )
+                {
+                    UE_LOG( LogStats, Log, TEXT( "Processing packet %i/%i" ), PacketIndex, Frame.Packets.Num() );
+                    PreviousSeconds = CurrentSeconds;
+                    bAtLeastOnePacket = true;
+                }
+            }
 
 			const FStatPacket& StatPacket = *Frame.Packets[PacketIndex];
 			const FName& ThreadFName = StatsThreadStats.Threads.FindChecked( StatPacket.ThreadId );
