Stats - Updated comment

diff --git a/Engine/Source/Programs/UnrealFrontend/Private/Commands/StatsDumpMemoryCommand.cpp b/Engine/Source/Programs/UnrealFrontend/Private/Commands/StatsDumpMemoryCommand.cpp
--- a/Engine/Source/Programs/UnrealFrontend/Private/Commands/StatsDumpMemoryCommand.cpp
+++ b/Engine/Source/Programs/UnrealFrontend/Private/Commands/StatsDumpMemoryCommand.cpp
@@ -77,8 +77,8 @@ void FStatsMemoryDumpCommand::InternalRun()
 		int64 TotalStatMessagesNum = 0;
 		int64 MaximumPacketSize = 0;
 		int64 TotalPacketsNum = 0;
-		// Read all packets sequentially, force by the memory profiler which is now a part of the raw stats.
-		// !!CAUTION!! Frame number in the raw stats is pointless, because it is time based, not frame based.
+		// Read all packets sequentially, forced by the memory profiler which is now a part of the raw stats.
+		// !!CAUTION!! Frame number in the raw stats is pointless, because it is time/cycles based, not frame based.
 		// Background threads usually execute time consuming operations, so the frame number won't be valid.
 		// Needs to be combined by the thread and the time, not by the frame number.
 		{
@@ -358,14 +358,14 @@ void FStatsMemoryDumpCommand::ProcessMemoryOperations( const TMap<int64, FStatPa
 	PreviousSeconds -= NumSecondsBetweenLogs;
 	for( int32 FrameIndex = 0; FrameIndex < Frames.Num(); ++FrameIndex )
 	{
-        {
-            const double CurrentSeconds = FPlatformTime::Seconds();
-            if( CurrentSeconds > PreviousSeconds + NumSecondsBetweenLogs )
-            {
-                UE_LOG( LogStats, Warning, TEXT( "Processing frame %i/%i" ), FrameIndex+1, Frames.Num() );
-                PreviousSeconds = CurrentSeconds;
-            }
-        }
+		{
+			const double CurrentSeconds = FPlatformTime::Seconds();
+			if( CurrentSeconds > PreviousSeconds + NumSecondsBetweenLogs )
+			{
+				UE_LOG( LogStats, Warning, TEXT( "Processing frame %i/%i" ), FrameIndex+1, Frames.Num() );
+				PreviousSeconds = CurrentSeconds;
+			}
+		}
 
 		const int64 TargetFrame = Frames[FrameIndex];
 		const int64 Diff = TargetFrame - FirstFrame;
@@ -374,15 +374,15 @@ void FStatsMemoryDumpCommand::ProcessMemoryOperations( const TMap<int64, FStatPa
 		bool bAtLeastOnePacket = false;
 		for( int32 PacketIndex = 0; PacketIndex < Frame.Packets.Num(); PacketIndex++ )
 		{
-            {
-                const double CurrentSeconds = FPlatformTime::Seconds();
-                if( CurrentSeconds > PreviousSeconds + NumSecondsBetweenLogs )
-                {
-                    UE_LOG( LogStats, Log, TEXT( "Processing packet %i/%i" ), PacketIndex, Frame.Packets.Num() );
-                    PreviousSeconds = CurrentSeconds;
-                    bAtLeastOnePacket = true;
-                }
-            }
+			{
+				const double CurrentSeconds = FPlatformTime::Seconds();
+				if( CurrentSeconds > PreviousSeconds + NumSecondsBetweenLogs )
+				{
+					UE_LOG( LogStats, Log, TEXT( "Processing packet %i/%i" ), PacketIndex, Frame.Packets.Num() );
+					PreviousSeconds = CurrentSeconds;
+					bAtLeastOnePacket = true;
+				}
+			}
 
 			const FStatPacket& StatPacket = *Frame.Packets[PacketIndex];
 			const FName& ThreadFName = StatsThreadStats.Threads.FindChecked( StatPacket.ThreadId );
