Profiler - Removed unused code

diff --git a/Engine/Source/Programs/UnrealFrontend/Private/Commands/StatsDumpMemoryCommand.cpp b/Engine/Source/Programs/UnrealFrontend/Private/Commands/StatsDumpMemoryCommand.cpp
--- a/Engine/Source/Programs/UnrealFrontend/Private/Commands/StatsDumpMemoryCommand.cpp
+++ b/Engine/Source/Programs/UnrealFrontend/Private/Commands/StatsDumpMemoryCommand.cpp
@@ -386,7 +386,6 @@ void FStatsMemoryDumpCommand::ProcessMemoryOperations( const TMap<int64, FStatPa
 
 			const FStatPacket& StatPacket = *Frame.Packets[PacketIndex];
 			const FName& ThreadFName = StatsThreadStats.Threads.FindChecked( StatPacket.ThreadId );
-			const uint32 NewThreadID = ThreadIDtoStatID.FindChecked( StatPacket.ThreadId );
 
 			FStackState* StackState = StackStates.Find( ThreadFName );
 			if( !StackState )
@@ -706,42 +705,6 @@ void FStatsMemoryDumpCommand::GenerateMemoryUsageReport( const TMap<uint64, FAll
 	}
 }
 
-void FStatsMemoryDumpCommand::CreateThreadsMapping()
-{
-	int32 UniqueID = 15;
-
-	for( auto It = StatsThreadStats.ShortNameToLongName.CreateConstIterator(); It; ++It )
-	{
-		const FStatMessage& LongName = It.Value();
-
-		const FName GroupName = LongName.NameAndInfo.GetGroupName();
-		
-		// Setup thread id to stat id.
-		if( GroupName == FStatConstants::NAME_ThreadGroup )
-		{
-			const FName StatName = It.Key();
-			UniqueID++;
-
-			uint32 ThreadID = 0;
-			for( auto ThreadsIt = StatsThreadStats.Threads.CreateConstIterator(); ThreadsIt; ++ThreadsIt )
-			{
-				if( ThreadsIt.Value() == StatName )
-				{
-					ThreadID = ThreadsIt.Key();
-					break;
-				}
-			}
-			ThreadIDtoStatID.Add( ThreadID, UniqueID );
-
-			// Game thread is always NAME_GameThread
-			if( StatName == NAME_GameThread )
-			{
-				GameThreadID = ThreadID;
-			}
-		}
-	}
-}
-
 void FStatsMemoryDumpCommand::ProcessingScopedAllocations( const TMap<uint64, FAllocationInfo>& AllocationMap )
 {
 	// This code is not optimized. 
diff --git a/Engine/Source/Programs/UnrealFrontend/Private/Commands/StatsDumpMemoryCommand.h b/Engine/Source/Programs/UnrealFrontend/Private/Commands/StatsDumpMemoryCommand.h
--- a/Engine/Source/Programs/UnrealFrontend/Private/Commands/StatsDumpMemoryCommand.h
+++ b/Engine/Source/Programs/UnrealFrontend/Private/Commands/StatsDumpMemoryCommand.h
@@ -47,13 +47,4 @@ protected:
 
 	/** Filepath to the raw stats file. */
 	FString SourceFilepath;
-
-	/**
-	 *	Helper map, as a part of migration into stats2, stored as ThreadID -> StatID.
-	 *	Used during creating a thread sample, so we can use the real stat id and get data graph for this thread
-	 */
-	TMap<uint32, uint32> ThreadIDtoStatID;
-
-	/** Game thread id. */
-	uint32 GameThreadID;
 };
