Compile fix

diff --git a/Engine/Source/Programs/UnrealFrontend/Private/Commands/StatsDumpMemoryCommand.cpp b/Engine/Source/Programs/UnrealFrontend/Private/Commands/StatsDumpMemoryCommand.cpp
--- a/Engine/Source/Programs/UnrealFrontend/Private/Commands/StatsDumpMemoryCommand.cpp
+++ b/Engine/Source/Programs/UnrealFrontend/Private/Commands/StatsDumpMemoryCommand.cpp
@@ -60,9 +60,6 @@ void FStatsMemoryDumpCommand::InternalRun()
 
 		const int64 CurrentFilePos = FileReader->Tell();
 
-		// Update profiler's metadata.
-		CreateThreadsMapping();
-
 		// Read frames offsets.
 		Stream.ReadFramesOffsets( *FileReader );
 
diff --git a/Engine/Source/Programs/UnrealFrontend/Private/Commands/StatsDumpMemoryCommand.h b/Engine/Source/Programs/UnrealFrontend/Private/Commands/StatsDumpMemoryCommand.h
--- a/Engine/Source/Programs/UnrealFrontend/Private/Commands/StatsDumpMemoryCommand.h
+++ b/Engine/Source/Programs/UnrealFrontend/Private/Commands/StatsDumpMemoryCommand.h
@@ -23,9 +23,6 @@ protected:
 	/** Executes the command. */
 	void InternalRun();
 
-	/** Creates thread mapping, some kind of metadata. */
-	void CreateThreadsMapping();
-
 	/** Basic memory profiling, only for debugging purpose. */
 	void ProcessMemoryOperations( const TMap<int64, FStatPacketArray>& CombinedHistory );
 
