Profiler - First pass for snapshots for memory profiling

diff --git a/Engine/Source/Programs/UnrealFrontend/Private/Commands/StatsDumpMemoryCommand.cpp b/Engine/Source/Programs/UnrealFrontend/Private/Commands/StatsDumpMemoryCommand.cpp
--- a/Engine/Source/Programs/UnrealFrontend/Private/Commands/StatsDumpMemoryCommand.cpp
+++ b/Engine/Source/Programs/UnrealFrontend/Private/Commands/StatsDumpMemoryCommand.cpp
@@ -4,6 +4,134 @@
 #include "StatsDumpMemoryCommand.h"
 #include "DiagnosticTable.h"
 
+/*-----------------------------------------------------------------------------
+	Callstack decoding/encoding
+-----------------------------------------------------------------------------*/
+
+static const TCHAR* CallstackSeparator = TEXT( "+" );
+
+static FString EncodeCallstack( const TArray<FName>& Callstack )
+{
+	FString Result;
+	for (const auto& Name : Callstack)
+	{
+		Result += TTypeToString<int32>::ToString( (int32)Name.GetComparisonIndex() );
+		Result += CallstackSeparator;
+	}
+	return Result;
+}
+
+static void DecodeCallstack( const FName& EncodedCallstack, TArray<FString>& out_DecodedCallstack )
+{
+	EncodedCallstack.ToString().ParseIntoArray( out_DecodedCallstack, CallstackSeparator, true );
+}
+
+static FString GetCallstack( const FName& EncodedCallstack )
+{
+	FString Result;
+
+	TArray<FString> DecodedCallstack;
+	DecodeCallstack( EncodedCallstack, DecodedCallstack );
+
+	for (int32 Index = DecodedCallstack.Num() - 1; Index >= 0; --Index)
+	{
+		NAME_INDEX NameIndex = 0;
+		TTypeFromString<NAME_INDEX>::FromString( NameIndex, *DecodedCallstack[Index] );
+
+		const FName LongName = FName( NameIndex, NameIndex, 0 );
+
+		const FString ShortName = FStatNameAndInfo::GetShortNameFrom( LongName ).ToString();
+		//const FString Group = FStatNameAndInfo::GetGroupNameFrom( LongName ).ToString();
+		FString Desc = FStatNameAndInfo::GetDescriptionFrom( LongName );
+		Desc.Trim();
+
+		if (Desc.Len() == 0)
+		{
+			Result += ShortName;
+		}
+		else
+		{
+			Result += Desc;
+		}
+
+		if (Index > 0)
+		{
+			Result += TEXT( " <- " );
+		}
+	}
+
+	Result.ReplaceInline( TEXT( "STAT_" ), TEXT( "" ), ESearchCase::CaseSensitive );
+	return Result;
+}
+
+/*-----------------------------------------------------------------------------
+	Allocation info
+-----------------------------------------------------------------------------*/
+
+FAllocationInfo::FAllocationInfo( uint64 InOldPtr, uint64 InPtr, int64 InSize, const TArray<FName>& InCallstack, uint32 InSequenceTag, EMemoryOperation InOp, bool bInHasBrokenCallstack ) 
+	: OldPtr( InOldPtr )
+	, Ptr( InPtr )
+	, Size( InSize )
+	, EncodedCallstack( *EncodeCallstack( InCallstack ) )
+	, SequenceTag( InSequenceTag )
+	, Op( InOp )
+	, bHasBrokenCallstack( bInHasBrokenCallstack )
+{
+
+}
+
+FAllocationInfo::FAllocationInfo( const FAllocationInfo& Other ) 
+	: OldPtr( Other.OldPtr )
+	, Ptr( Other.Ptr )
+	, Size( Other.Size )
+	, EncodedCallstack( Other.EncodedCallstack )
+	, SequenceTag( Other.SequenceTag )
+	, Op( Other.Op )
+	, bHasBrokenCallstack( Other.bHasBrokenCallstack )
+{
+
+}
+
+struct FAllocationInfoGreater
+{
+	FORCEINLINE bool operator()( const FAllocationInfo& A, const FAllocationInfo& B ) const
+	{
+		return B.Size < A.Size;
+	}
+};
+struct FSizeAndCountGreater
+{
+	FORCEINLINE bool operator()( const FSizeAndCount& A, const FSizeAndCount& B ) const
+	{
+		return B.Size < A.Size;
+	}
+};
+
+/*-----------------------------------------------------------------------------
+	Stats stack helpers
+-----------------------------------------------------------------------------*/
+
+/** Holds stats stack state, used to preserve continuity when the game frame has changed. */
+struct FStackState
+{
+	FStackState()
+		: bIsBrokenCallstack( false )
+	{}
+
+	/** Call stack. */
+	TArray<FName> Stack;
+
+	/** Current function name. */
+	FName Current;
+
+	/** Whether this callstack is marked as broken due to mismatched start and end scope cycles. */
+	bool bIsBrokenCallstack;
+};
+
+/*-----------------------------------------------------------------------------
+	FStatsMemoryDumpCommand
+-----------------------------------------------------------------------------*/
+
 static const double NumSecondsBetweenLogs = 5.0;
 
 void FStatsMemoryDumpCommand::InternalRun()
@@ -33,7 +161,7 @@ void FStatsMemoryDumpCommand::InternalRun()
 
 	const bool bIsFinalized = Stream.Header.IsFinalized();
 	check( bIsFinalized );
-	check( Stream.Header.Version == EStatMagicWithHeader::VERSION_5 );
+	check( Stream.Header.Version >= EStatMagicWithHeader::VERSION_6 );
 	StatsThreadStats.MarkAsLoaded();
 
 	TArray<FStatMessage> Messages;
@@ -49,12 +177,13 @@ void FStatsMemoryDumpCommand::InternalRun()
 		// Find all UObject metadata messages.
 		for( const auto& Meta : MetadataMessages )
 		{
-			FName LongName = Meta.NameAndInfo.GetRawName();
-			const FString Desc = FStatNameAndInfo::GetShortNameFrom( LongName ).GetPlainNameString();
+			const FName EncName = Meta.NameAndInfo.GetEncodedName();
+			const FName RawName = Meta.NameAndInfo.GetRawName();
+			const FString Desc = FStatNameAndInfo::GetShortNameFrom( RawName ).GetPlainNameString();
 			const bool bContainsUObject = Desc.Contains( TEXT( "//" ) );
 			if( bContainsUObject )
 			{
-				UObjectNames.Add( LongName );
+				UObjectNames.Add( RawName );
 			}
 		}
 
@@ -174,161 +303,6 @@ void FStatsMemoryDumpCommand::InternalRun()
 }
 
 
-static const TCHAR* CallstackSeparator = TEXT( "+" );
-
-static FString EncodeCallstack( const TArray<FName>& Callstack )
-{
-	FString Result;
-	for( const auto& Name : Callstack )
-	{
-		Result += TTypeToString<int32>::ToString( (int32)Name.GetComparisonIndex() );
-		Result += CallstackSeparator;
-	}
-	return Result;
-}
-
-static void DecodeCallstack( const FName& EncodedCallstack, TArray<FString>& out_DecodedCallstack )
-{
-	EncodedCallstack.ToString().ParseIntoArray( out_DecodedCallstack, CallstackSeparator, true );
-}
-
-/** Simple allocation info. Assumes the sequence tag are unique and doesn't wrap. */
-struct FAllocationInfo
-{
-	uint64 Ptr;
-	int64 Size;
-	FName EncodedCallstack;
-	uint32 SequenceTag;
-	EMemoryOperation Op;
-	bool bHasBrokenCallstack;
-
-	FAllocationInfo(
-		uint64 InPtr,
-		int64 InSize,
-		const TArray<FName>& InCallstack,
-		uint32 InSequenceTag,
-		EMemoryOperation InOp,
-		bool bInHasBrokenCallstack
-		)
-		: Ptr( InPtr )
-		, Size( InSize )
-		, EncodedCallstack( *EncodeCallstack( InCallstack ) )
-		, SequenceTag( InSequenceTag )
-		, Op( InOp )
-		, bHasBrokenCallstack( bInHasBrokenCallstack )
-	{}
-
-	FAllocationInfo( const FAllocationInfo& Other )
-		: Ptr( Other.Ptr )
-		, Size( Other.Size )
-		, EncodedCallstack( Other.EncodedCallstack )
-		, SequenceTag( Other.SequenceTag )
-		, Op( Other.Op )
-		, bHasBrokenCallstack( Other.bHasBrokenCallstack )
-	{}
-
-// 	FAllocationInfo()
-// 		: Ptr( 0 )
-// 		, Size( 0 )
-// 		, SequenceTag( 0 )
-// 		, Op( EMemoryOperation::Invalid )
-// 		, bHasBrokenCallstack( false )
-// 	{}
-
-
-	bool operator<(const FAllocationInfo& Other) const
-	{
-		return SequenceTag < Other.SequenceTag;
-	}
-
-	bool operator==(const FAllocationInfo& Other) const
-	{
-		return SequenceTag == Other.SequenceTag;
-	}
-};
-
-struct FAllocationInfoGreater
-{
-	FORCEINLINE bool operator()( const FAllocationInfo& A, const FAllocationInfo& B ) const
-	{
-		return B.Size < A.Size;
-	}
-};
-
-
-struct FSizeAndCount
-{
-	uint64 Size;
-	uint64 Count;
-	FSizeAndCount()
-		: Size( 0 )
-		, Count( 0 )
-	{}
-};
-
-struct FSizeAndCountGreater
-{
-	FORCEINLINE bool operator()( const FSizeAndCount& A, const FSizeAndCount& B ) const
-	{
-		return B.Size < A.Size;
-	}
-};
-
-/** Holds stats stack state, used to preserve continuity when the game frame has changed. */
-struct FStackState
-{
-	FStackState()
-		: bIsBrokenCallstack( false )
-	{}
-
-	/** Call stack. */
-	TArray<FName> Stack;
-
-	/** Current function name. */
-	FName Current;
-
-	/** Whether this callstack is marked as broken due to mismatched start and end scope cycles. */
-	bool bIsBrokenCallstack;
-};
-
-static FString GetCallstack( const FName& EncodedCallstack )
-{
-	FString Result;
-
-	TArray<FString> DecodedCallstack;
-	DecodeCallstack( EncodedCallstack, DecodedCallstack );
-
-	for( int32 Index = DecodedCallstack.Num() - 1; Index >= 0; --Index )
-	{
-		NAME_INDEX NameIndex = 0;
-		TTypeFromString<NAME_INDEX>::FromString( NameIndex, *DecodedCallstack[Index] );
-
-		const FName LongName = FName( NameIndex, NameIndex, 0 );
-
-		const FString ShortName = FStatNameAndInfo::GetShortNameFrom( LongName ).ToString();
-		//const FString Group = FStatNameAndInfo::GetGroupNameFrom( LongName ).ToString();
-		FString Desc = FStatNameAndInfo::GetDescriptionFrom( LongName );
-		Desc.Trim();
-
-		if( Desc.Len() == 0 )
-		{
-			Result += ShortName;
-		}
-		else
-		{
-			Result += Desc;
-		}
-
-		if( Index > 0 )
-		{
-			Result += TEXT( " <- " );
-		}
-	}
-
-	Result.ReplaceInline( TEXT( "STAT_" ), TEXT( "" ), ESearchCase::CaseSensitive );
-	return Result;
-}
-
 void FStatsMemoryDumpCommand::ProcessMemoryOperations( const TMap<int64, FStatPacketArray>& CombinedHistory )
 {
 	// This is only example code, no fully implemented, may sometimes crash.
@@ -353,6 +327,11 @@ void FStatsMemoryDumpCommand::ProcessMemoryOperations( const TMap<int64, FStatPa
 	// Read all stats messages, parse all memory operations and decode callstacks.
 	const int64 FirstFrame = 0;
 	PreviousSeconds -= NumSecondsBetweenLogs;
+
+	// Begin marker.
+	uint32 LastSequenceTagForNamedMarker = 0;
+	Snapshots.Add( TPairInitializer<uint32, FName>( LastSequenceTagForNamedMarker, TEXT( "BeginSnapshot" ) ) );
+
 	for( int32 FrameIndex = 0; FrameIndex < Frames.Num(); ++FrameIndex )
 	{
 		{
@@ -417,7 +396,7 @@ void FStatsMemoryDumpCommand::ProcessMemoryOperations( const TMap<int64, FStatPa
 				const EStatOperation::Type Op = Item.NameAndInfo.GetField<EStatOperation>();
 				const FName RawName = Item.NameAndInfo.GetRawName();
 
-				if( Op == EStatOperation::CycleScopeStart || Op == EStatOperation::CycleScopeEnd || Op == EStatOperation::Memory )
+				if (Op == EStatOperation::CycleScopeStart || Op == EStatOperation::CycleScopeEnd || Op == EStatOperation::Memory || Op == EStatOperation::SpecialMessageMarker)
 				{
 					if( Op == EStatOperation::CycleScopeStart )
 					{
@@ -429,26 +408,29 @@ void FStatsMemoryDumpCommand::ProcessMemoryOperations( const TMap<int64, FStatPa
 						// Experimental code used only to test the implementation.
 						// First memory operation is Alloc or Free
 						const uint64 EncodedPtr = Item.GetValue_Ptr();
-						const bool bIsAlloc = (EncodedPtr & (uint64)EMemoryOperation::Alloc) != 0;
-						const bool bIsFree = (EncodedPtr & (uint64)EMemoryOperation::Free) != 0;
+						const EMemoryOperation MemOp = EMemoryOperation( EncodedPtr & (uint64)EMemoryOperation::Mask );
 						const uint64 Ptr = EncodedPtr & ~(uint64)EMemoryOperation::Mask;
-						if( bIsAlloc )
+						if (MemOp == EMemoryOperation::Alloc)
 						{
 							NumMemoryOperations++;
 							// @see FStatsMallocProfilerProxy::TrackAlloc
-							// After alloc ptr message there is always alloc size message and the sequence tag.
+							// After AllocPtr message there is always alloc size message and the sequence tag.
 							Index++;
 							const FStatMessage& AllocSizeMessage = Data[Index];
 							const int64 AllocSize = AllocSizeMessage.GetValue_int64();
 
-							// Read operation sequence tag.
+							// Read OperationSequenceTag.
 							Index++;
 							const FStatMessage& SequenceTagMessage = Data[Index];
 							const uint32 SequenceTag = SequenceTagMessage.GetValue_int64();
 
+							//ThreadStats->AddMemoryMessage( GET_STATFNAME( STAT_Memory_AllocPtr ), (uint64)(UPTRINT)Ptr | (uint64)EMemoryOperation::Alloc );
+							//ThreadStats->AddMemoryMessage( GET_STATFNAME( STAT_Memory_AllocSize ), Size );
+							//ThreadStats->AddMemoryMessage( GET_STATFNAME( STAT_Memory_OperationSequenceTag ), (int64)SequenceTag );
+
 							// Create a callstack.
 							TArray<FName> StatsBasedCallstack;
-							for( const auto& StackName : StackState->Stack )
+							for (const auto& StackName : StackState->Stack)
 							{
 								StatsBasedCallstack.Add( StackName );
 							}
@@ -456,6 +438,7 @@ void FStatsMemoryDumpCommand::ProcessMemoryOperations( const TMap<int64, FStatPa
 							// Add a new allocation.
 							SequenceAllocationArray.Add(
 								FAllocationInfo(
+								0,
 								Ptr,
 								AllocSize,
 								StatsBasedCallstack,
@@ -463,34 +446,83 @@ void FStatsMemoryDumpCommand::ProcessMemoryOperations( const TMap<int64, FStatPa
 								EMemoryOperation::Alloc,
 								StackState->bIsBrokenCallstack
 								) );
+							LastSequenceTagForNamedMarker = SequenceTag;
+						}
+						else if (MemOp == EMemoryOperation::Realloc)
+						{
+							NumMemoryOperations++;
+							const uint64 OldPtr = Ptr;
+
+							// Read NewPtr
+							Index++;
+							const FStatMessage& AllocPtrMessage = Data[Index];
+							const uint64 NewPtr = AllocPtrMessage.GetValue_Ptr() & ~(uint64)EMemoryOperation::Mask;
+
+							// After AllocPtr message there is always alloc size message and the sequence tag.
+							Index++;
+							const FStatMessage& ReallocSizeMessage = Data[Index];
+							const int64 ReallocSize = ReallocSizeMessage.GetValue_int64();
+
+							// Read OperationSequenceTag.
+							Index++;
+							const FStatMessage& SequenceTagMessage = Data[Index];
+							const uint32 SequenceTag = SequenceTagMessage.GetValue_int64();
+
+							//ThreadStats->AddMemoryMessage( GET_STATFNAME( STAT_Memory_FreePtr ), (uint64)(UPTRINT)OldPtr | (uint64)EMemoryOperation::Realloc );
+							//ThreadStats->AddMemoryMessage( GET_STATFNAME( STAT_Memory_AllocPtr ), (uint64)(UPTRINT)NewPtr | (uint64)EMemoryOperation::Realloc );
+							//ThreadStats->AddMemoryMessage( GET_STATFNAME( STAT_Memory_AllocSize ), NewSize );
+							//ThreadStats->AddMemoryMessage( GET_STATFNAME( STAT_Memory_OperationSequenceTag ), (int64)SequenceTag );
+
+							// Create a callstack.
+							TArray<FName> StatsBasedCallstack;
+							for (const auto& StackName : StackState->Stack)
+							{
+								StatsBasedCallstack.Add( StackName );
+							}
+
+							// Add a new realloc.
+							SequenceAllocationArray.Add(
+								FAllocationInfo(
+								OldPtr, 
+								NewPtr,
+								ReallocSize,
+								StatsBasedCallstack,
+								SequenceTag,
+								EMemoryOperation::Realloc,
+								StackState->bIsBrokenCallstack
+								) );
+							LastSequenceTagForNamedMarker = SequenceTag;
 						}
-						else if( bIsFree )
+						else if (MemOp == EMemoryOperation::Free)
 						{
 							NumMemoryOperations++;
-							// Read operation sequence tag.
+							// Read OperationSequenceTag.
 							Index++;
 							const FStatMessage& SequenceTagMessage = Data[Index];
 							const uint32 SequenceTag = SequenceTagMessage.GetValue_int64();
 
+							//ThreadStats->AddMemoryMessage( GET_STATFNAME( STAT_Memory_FreePtr ), (uint64)(UPTRINT)Ptr | (uint64)EMemoryOperation::Free );	// 16 bytes total				
+							//ThreadStats->AddMemoryMessage( GET_STATFNAME( STAT_Memory_OperationSequenceTag ), (int64)SequenceTag );
+
 							// Create a callstack.
-							/*
 							TArray<FName> StatsBasedCallstack;
-							for( const auto& RawName : StackState->Stack )
+							for (const auto& StackName : StackState->Stack)
 							{
-								StatsBasedCallstack.Add( RawName );
+								StatsBasedCallstack.Add( StackName );
 							}
-							*/
 
 							// Add a new free.
 							SequenceAllocationArray.Add(
 								FAllocationInfo(
+								0,
 								Ptr,		
 								0,
-								TArray<FName>()/*StatsBasedCallstack*/,					
+								StatsBasedCallstack,					
 								SequenceTag,
 								EMemoryOperation::Free,
 								StackState->bIsBrokenCallstack
 								) );
+							//LastSequenceTagForNamedMarker = SequenceTag;????
 						}
 						else
 						{
@@ -529,6 +561,14 @@ void FStatsMemoryDumpCommand::ProcessMemoryOperations( const TMap<int64, FStatPa
 							StackState->Current = ThreadFName;
 						}
 					}
+					else if (Op == EStatOperation::SpecialMessageMarker)
+					{
+						if (RawName == FStatConstants::NAME_NamedMarker)
+						{
+							const FName NamedMarker = Item.GetValue_FName(); // LastSequenceTag
+							Snapshots.Add( TPairInitializer<uint32, FName>( LastSequenceTagForNamedMarker, NamedMarker ) );
+						}
+					}
 				}
 			}
 			if( bAtLeastOneMessage )
@@ -542,6 +582,12 @@ void FStatsMemoryDumpCommand::ProcessMemoryOperations( const TMap<int64, FStatPa
 		}
 	}
 
+	// End marker.
+	Snapshots.Add( TPairInitializer<uint32, FName>( TNumericLimits<uint32>::Max(), TEXT( "EndSnapshot" ) ) );
+
+	// Copy snapshots.
+	SnapshotsToBeProcessed = Snapshots;
+
 	UE_LOG( LogStats, Warning, TEXT( "NumMemoryOperations:   %llu" ), NumMemoryOperations );
 	UE_LOG( LogStats, Warning, TEXT( "SequenceAllocationNum: %i" ), SequenceAllocationArray.Num() );
 
@@ -568,12 +614,23 @@ void FStatsMemoryDumpCommand::ProcessMemoryOperations( const TMap<int64, FStatPa
 	// Sort all memory operation by the sequence tag, iterate through all operation and generate memory usage.
 	SequenceAllocationArray.Sort( TLess<FAllocationInfo>() );
 
+	// Named markers/snapshots
+
 	// Alive allocations.
 	TMap<uint64, FAllocationInfo> AllocationMap;
 	TMultiMap<uint64, FAllocationInfo> FreeWithoutAllocMap;
 	TMultiMap<uint64, FAllocationInfo> DuplicatedAllocMap;
 	int32 NumDuplicatedMemoryOperations = 0;
 	int32 NumFWAMemoryOperations = 0; // FreeWithoutAlloc
+	int32 NumZeroAllocs = 0; // Malloc(0)
+
+	// Initialize the begin snapshot.
+	auto BeginSnapshot = SnapshotsToBeProcessed[0];
+	SnapshotsToBeProcessed.RemoveAt( 0 );
+	PrepareSnapshot( BeginSnapshot.Value, AllocationMap );
+
+	
+	auto CurrentSnapshot = SnapshotsToBeProcessed[0];
 
 	UE_LOG( LogStats, Warning, TEXT( "Generating memory operations map" ) );
 	const int32 NumSequenceAllocations = SequenceAllocationArray.Num();
@@ -592,25 +649,36 @@ void FStatsMemoryDumpCommand::ProcessMemoryOperations( const TMap<int64, FStatPa
 		}
 
 		const FAllocationInfo& Alloc = SequenceAllocationArray[Index];
-		const EMemoryOperation MemOp = Alloc.Op;
-		const uint64 Ptr = Alloc.Ptr;
-		const int64 Size = Alloc.Size;
-		const uint32 SequenceTag = Alloc.SequenceTag;
 
-		if( MemOp == EMemoryOperation::Alloc )
+		// Check named marker/snapshots
+		if (Alloc.SequenceTag > CurrentSnapshot.Key)
 		{
-			const FAllocationInfo* Found = AllocationMap.Find( Ptr );
+			SnapshotsToBeProcessed.RemoveAt( 0 );
+			PrepareSnapshot( CurrentSnapshot.Value, AllocationMap );
+			CurrentSnapshot = SnapshotsToBeProcessed[0];
+		}
+
+		
+
+		if (Alloc.Op == EMemoryOperation::Alloc)
+		{
+			if (Alloc.Size == 0)
+			{
+				NumZeroAllocs++;
+			}
+
+			const FAllocationInfo* Found = AllocationMap.Find( Alloc.Ptr );
 
 			if( !Found )
 			{
-				AllocationMap.Add( Ptr, Alloc );
+				AllocationMap.Add( Alloc.Ptr, Alloc );
 			}
 			else
 			{
 				const FAllocationInfo* FoundAndFreed = FreeWithoutAllocMap.Find( Found->Ptr );
 				const FAllocationInfo* FoundAndAllocated = FreeWithoutAllocMap.Find( Alloc.Ptr );
 
-#if	_DEBUG
+#if	UE_BUILD_DEBUG
 				if( FoundAndFreed )
 				{
 					const FString FoundAndFreedCallstack = GetCallstack( FoundAndFreed->EncodedCallstack );
@@ -626,36 +694,85 @@ void FStatsMemoryDumpCommand::ProcessMemoryOperations( const TMap<int64, FStatPa
 
 				const FString FoundCallstack = GetCallstack( Found->EncodedCallstack );
 				const FString AllocCallstack = GetCallstack( Alloc.EncodedCallstack );
-#endif // _DEBUG
+				UE_LOG( LogStats, Warning, TEXT( "DuplicateAlloc: %s Size: %i/%i Ptr: %i/%i Tag: %i/%i" ), *AllocCallstack, Found->Size, Alloc.Size, Found->Ptr, Alloc.Ptr, Found->SequenceTag, Alloc.SequenceTag );
+#endif // UE_BUILD_DEBUG
 
 				// Replace pointer.
-				AllocationMap.Add( Ptr, Alloc );
+				AllocationMap.Add( Alloc.Ptr, Alloc );
 				// Store the old pointer.
-				DuplicatedAllocMap.Add( Ptr, *Found );
+				DuplicatedAllocMap.Add( Alloc.Ptr, *Found );
+			}
+		}
+		else if (Alloc.Op == EMemoryOperation::Realloc)
+		{
+			if (Alloc.Size == 0)
+			{
+				NumZeroAllocs++;
 			}
+
+			// Previous Alloc or Realloc
+			const FAllocationInfo* FoundOld = AllocationMap.Find( Alloc.OldPtr );
+
+			if (FoundOld)
+			{
+				const bool bIsValid = Alloc.SequenceTag > FoundOld->SequenceTag;
+				if (!bIsValid)
+				{
+					UE_LOG( LogStats, Warning, TEXT( "InvalidRealloc Ptr: %llu, Seq: %i/%i" ), Alloc.Ptr, Alloc.SequenceTag, FoundOld->SequenceTag );
+				}
+				// Remove the old allocation.
+				AllocationMap.Remove( Alloc.OldPtr );
+				AllocationMap.Add( Alloc.Ptr, Alloc );
+			}
+			// If we have situation where the old pointer is the same as the new pointer
+			// There is high change that the original call looked like this 
+			// Ptr = Malloc( 40 ); Ptr = Realloc( Ptr, 40 );
+			else if (Alloc.OldPtr != Alloc.Ptr)
+			{
+				// No OldPtr, should not happen as it means realloc without initial alloc.
+				// Or Realloc after Malloc(0)
+
+#if	UE_BUILD_DEBUG
+				const FString ReallocCallstack = GetCallstack( Alloc.EncodedCallstack );
+				UE_LOG( LogStats, Warning, TEXT( "ReallocWithoutAlloc: %s %i %i/%i [%i]" ), *ReallocCallstack, Alloc.Size, Alloc.OldPtr, Alloc.Ptr, Alloc.SequenceTag );
+#endif // UE_BUILD_DEBUG
+
+				AllocationMap.Add( Alloc.Ptr, Alloc );
+			}
+
 		}
-		else if( MemOp == EMemoryOperation::Free )
+		else if (Alloc.Op == EMemoryOperation::Free)
 		{
-			const FAllocationInfo* Found = AllocationMap.Find( Ptr );
+			const FAllocationInfo* Found = AllocationMap.Find( Alloc.Ptr );
 			if( Found )
 			{
 				const bool bIsValid = Alloc.SequenceTag > Found->SequenceTag;
 				if( !bIsValid )
 				{
-					UE_LOG( LogStats, Warning, TEXT( "InvalidFree Ptr: %llu, Seq: %i/%i" ), Ptr, SequenceTag, Found->SequenceTag );
+					UE_LOG( LogStats, Warning, TEXT( "InvalidFree Ptr: %llu, Seq: %i/%i" ), Alloc.Ptr, Alloc.SequenceTag, Found->SequenceTag );
 				}
-				AllocationMap.Remove( Ptr );
+				AllocationMap.Remove( Alloc.Ptr );
 			}
 			else
 			{
-				FreeWithoutAllocMap.Add( Ptr, Alloc );
+				FreeWithoutAllocMap.Add( Alloc.Ptr, Alloc );
 				NumFWAMemoryOperations++;
+
+#if	UE_BUILD_DEBUG
+				const FString FWACallstack = GetCallstack( Alloc.EncodedCallstack );
+				UE_LOG( LogStats, Warning, TEXT( "FreeWithoutAllocCallstack: %s %i" ), *FWACallstack, Alloc.Ptr );
+#endif // UE_BUILD_DEBUG
 			}
 		}
 	}
 
+	auto EndSnapshot = SnapshotsToBeProcessed[0];
+	SnapshotsToBeProcessed.RemoveAt( 0 );
+	PrepareSnapshot( EndSnapshot.Value, AllocationMap );
+
 	UE_LOG( LogStats, Warning, TEXT( "NumDuplicatedMemoryOperations: %i" ), NumDuplicatedMemoryOperations );
 	UE_LOG( LogStats, Warning, TEXT( "NumFWAMemoryOperations:        %i" ), NumFWAMemoryOperations );
+	UE_LOG( LogStats, Warning, TEXT( "NumZeroAllocs:                 %i" ), NumZeroAllocs );
 
 	// Dump problematic allocations
 	DuplicatedAllocMap.ValueSort( FAllocationInfoGreater() );
@@ -686,7 +803,21 @@ void FStatsMemoryDumpCommand::ProcessMemoryOperations( const TMap<int64, FStatPa
 		}
 	}
 
-	GenerateMemoryUsageReport( AllocationMap );
+	// Frame-240 Frame-120 Frame-060
+	TMap<FString, FSizeAndCount> FrameBegin_Exit;
+	CompareSnapshots_FString( TEXT( "BeginSnapshot" ), TEXT( "EngineLoop.Exit" ), FrameBegin_Exit );
+	DumpScopedAllocations( TEXT( "FrameBegin_Exit" ), FrameBegin_Exit );
+
+	TMap<FString, FSizeAndCount> Frame060_120;
+	CompareSnapshots_FString( TEXT( "Frame-060" ), TEXT( "Frame-120" ), Frame060_120 );
+	DumpScopedAllocations( TEXT( "Frame060_120" ), Frame060_120 );
+
+	TMap<FString, FSizeAndCount> Frame060_240;
+	CompareSnapshots_FString( TEXT( "Frame-060" ), TEXT( "Frame-240" ), Frame060_240 );
+	DumpScopedAllocations( TEXT( "Frame060_240" ), Frame060_240 );
+
+	// For snapshot?
+	//GenerateMemoryUsageReport( AllocationMap );
 }
 
 void FStatsMemoryDumpCommand::GenerateMemoryUsageReport( const TMap<uint64, FAllocationInfo>& AllocationMap )
@@ -697,19 +828,20 @@ void FStatsMemoryDumpCommand::GenerateMemoryUsageReport( const TMap<uint64, FAll
 	}
 	else
 	{
-		ProcessingUObjectAllocations( AllocationMap );
-		ProcessingScopedAllocations( AllocationMap );
+		//ProcessAndDumpUObjectAllocations( AllocationMap );
+		ProcessAndDumpScopedAllocations( AllocationMap );
 	}
 }
 
-void FStatsMemoryDumpCommand::ProcessingScopedAllocations( const TMap<uint64, FAllocationInfo>& AllocationMap )
+void FStatsMemoryDumpCommand::ProcessAndDumpScopedAllocations( const TMap<uint64, FAllocationInfo>& AllocationMap )
 {
 	// This code is not optimized. 
 	FScopeLogTime SLT( TEXT( "ProcessingScopedAllocations" ), nullptr, FScopeLogTime::ScopeLog_Seconds );
-
 	UE_LOG( LogStats, Warning, TEXT( "Processing scoped allocations" ) );
 
-	FDiagnosticTableViewer MemoryReport( *FDiagnosticTableViewer::GetUniqueTemporaryFilePath( TEXT( "MemoryReport-Scoped" ) ) );
+	const FString ReportName = FString::Printf( TEXT( "%s-Memory-Scoped" ), *Stream.Header.PlatformName );
+	FDiagnosticTableViewer MemoryReport( *FDiagnosticTableViewer::GetUniqueTemporaryFilePath( *ReportName ), true );
+
 	// Write a row of headings for the table's columns.
 	MemoryReport.AddColumn( TEXT( "Size (bytes)" ) );
 	MemoryReport.AddColumn( TEXT( "Size (MB)" ) );
@@ -717,21 +849,10 @@ void FStatsMemoryDumpCommand::ProcessingScopedAllocations( const TMap<uint64, FA
 	MemoryReport.AddColumn( TEXT( "Callstack" ) );
 	MemoryReport.CycleRow();
 
-
 	TMap<FName, FSizeAndCount> ScopedAllocations;
-
-	uint64 NumAllocations = 0;
 	uint64 TotalAllocatedMemory = 0;
-	for( const auto& It : AllocationMap )
-	{
-		const FAllocationInfo& Alloc = It.Value;
-		FSizeAndCount& SizeAndCount = ScopedAllocations.FindOrAdd( Alloc.EncodedCallstack );
-		SizeAndCount.Size += Alloc.Size;
-		SizeAndCount.Count += 1;
-
-		TotalAllocatedMemory += Alloc.Size;
-		NumAllocations++;
-	}
+	uint64 NumAllocations = 0;
+	GenerateScopedAllocations( AllocationMap, ScopedAllocations, TotalAllocatedMemory, NumAllocations );
 
 	// Dump memory to the log.
 	ScopedAllocations.ValueSort( FSizeAndCountGreater() );
@@ -784,13 +905,15 @@ void FStatsMemoryDumpCommand::ProcessingScopedAllocations( const TMap<uint64, FA
 	MemoryReport.CycleRow();
 }
 
-void FStatsMemoryDumpCommand::ProcessingUObjectAllocations( const TMap<uint64, FAllocationInfo>& AllocationMap )
+void FStatsMemoryDumpCommand::ProcessAndDumpUObjectAllocations( const TMap<uint64, FAllocationInfo>& AllocationMap )
 {
 	// This code is not optimized. 
 	FScopeLogTime SLT( TEXT( "ProcessingUObjectAllocations" ), nullptr, FScopeLogTime::ScopeLog_Seconds );
 	UE_LOG( LogStats, Warning, TEXT( "Processing UObject allocations" ) );
 
-	FDiagnosticTableViewer MemoryReport( *FDiagnosticTableViewer::GetUniqueTemporaryFilePath( TEXT( "MemoryReport-UObject" ) ) );
+	const FString ReportName = FString::Printf( TEXT("%s-Memory-UObject"), *Stream.Header.PlatformName );
+	FDiagnosticTableViewer MemoryReport( *FDiagnosticTableViewer::GetUniqueTemporaryFilePath( *ReportName ), true );
+
 	// Write a row of headings for the table's columns.
 	MemoryReport.AddColumn( TEXT( "Size (bytes)" ) );
 	MemoryReport.AddColumn( TEXT( "Size (MB)" ) );
@@ -835,8 +958,7 @@ void FStatsMemoryDumpCommand::ProcessingUObjectAllocations( const TMap<uint64, F
 		if( UObjectClass != NAME_None )
 		{
 			FSizeAndCount& SizeAndCount = UObjectAllocations.FindOrAdd( UObjectClass );
-			SizeAndCount.Size += Alloc.Size;
-			SizeAndCount.Count += 1;
+			SizeAndCount.AddAllocation( Alloc.Size );
 
 			TotalAllocatedMemory += Alloc.Size;
 			NumAllocations++;
@@ -891,3 +1013,179 @@ void FStatsMemoryDumpCommand::ProcessingUObjectAllocations( const TMap<uint64, F
 	MemoryReport.AddColumn( TEXT( "TOTAL" ) );
 	MemoryReport.CycleRow();
 }
+
+void FStatsMemoryDumpCommand::DumpScopedAllocations( const TCHAR* Name, const TMap<FString, FSizeAndCount>& ScopedAllocations )
+{
+	// This code is not optimized. 
+	FScopeLogTime SLT( TEXT( "ProcessingScopedAllocations" ), nullptr, FScopeLogTime::ScopeLog_Seconds );
+	UE_LOG( LogStats, Warning, TEXT( "Processing scoped allocations" ) );
+
+	const FString ReportName = FString::Printf( TEXT( "%s-Memory-Scoped-%s" ), *Stream.Header.PlatformName, Name );
+	FDiagnosticTableViewer MemoryReport( *FDiagnosticTableViewer::GetUniqueTemporaryFilePath( *ReportName ), true );
+
+	// Write a row of headings for the table's columns.
+	MemoryReport.AddColumn( TEXT( "Size (bytes)" ) );
+	MemoryReport.AddColumn( TEXT( "Size (MB)" ) );
+	MemoryReport.AddColumn( TEXT( "Count" ) );
+	MemoryReport.AddColumn( TEXT( "Callstack" ) );
+	MemoryReport.CycleRow();
+
+	
+	FSizeAndCount Total;
+
+	const float MaxPctDisplayed = 0.90f;
+	int32 CurrentIndex = 0;
+	UE_LOG( LogStats, Warning, TEXT( "Index, Size (Size MB), Count, Stat desc" ) );
+	for (const auto& It : ScopedAllocations)
+	{
+		const FSizeAndCount& SizeAndCount = It.Value;
+		//const FName& EncodedCallstack = It.Key;
+		const FString AllocCallstack = It.Key;// GetCallstack( EncodedCallstack );
+
+		UE_LOG( LogStats, Log, TEXT( "%2i, %llu (%.2f MB), %llu, %s" ),
+				CurrentIndex,
+				SizeAndCount.Size,
+				SizeAndCount.Size / 1024.0f / 1024.0f,
+				SizeAndCount.Count,
+				*AllocCallstack );
+
+		// Dump stats
+		MemoryReport.AddColumn( TEXT( "%llu" ), SizeAndCount.Size );
+		MemoryReport.AddColumn( TEXT( "%.2f MB" ), SizeAndCount.Size / 1024.0f / 1024.0f );
+		MemoryReport.AddColumn( TEXT( "%llu" ), SizeAndCount.Count );
+		MemoryReport.AddColumn( *AllocCallstack );
+		MemoryReport.CycleRow();
+
+		CurrentIndex++;
+		Total += SizeAndCount;
+	}
+
+	UE_LOG( LogStats, Warning, TEXT( "Allocated memory: %llu bytes (%.2f MB)" ), Total.Size, Total.SizeMB );
+
+	// Add a total row.
+	MemoryReport.CycleRow();
+	MemoryReport.CycleRow();
+	MemoryReport.CycleRow();
+	MemoryReport.AddColumn( TEXT( "%llu" ), Total.Size );
+	MemoryReport.AddColumn( TEXT( "%.2f MB" ), Total.SizeMB );
+	MemoryReport.AddColumn( TEXT( "%llu" ), Total.Count );
+	MemoryReport.AddColumn( TEXT( "TOTAL" ) );
+	MemoryReport.CycleRow();
+}
+
+void FStatsMemoryDumpCommand::GenerateScopedAllocations( const TMap<uint64, FAllocationInfo>& AllocationMap, TMap<FName, FSizeAndCount>& out_ScopedAllocations, uint64& TotalAllocatedMemory, uint64& NumAllocations )
+{
+	FScopeLogTime SLT( TEXT( "GenerateScopedAllocations" ), nullptr, FScopeLogTime::ScopeLog_Milliseconds );
+
+	for (const auto& It : AllocationMap)
+	{
+		const FAllocationInfo& Alloc = It.Value;
+		FSizeAndCount& SizeAndCount = out_ScopedAllocations.FindOrAdd( Alloc.EncodedCallstack );
+		SizeAndCount.AddAllocation( Alloc.Size );
+
+		TotalAllocatedMemory += Alloc.Size;
+		NumAllocations++;
+	}
+
+	// Sort by size.
+	out_ScopedAllocations.ValueSort( FSizeAndCountGreater() );
+}
+
+void FStatsMemoryDumpCommand::PrepareSnapshot( const FName SnapshotName, const TMap<uint64, FAllocationInfo>& AllocationMap )
+{
+	FScopeLogTime SLT( TEXT( "PrepareSnapshot" ), nullptr, FScopeLogTime::ScopeLog_Milliseconds );
+
+	SnapshotsWithAllocationMap.Add( SnapshotName, AllocationMap );
+
+	TMap<FName, FSizeAndCount> SnapshotScopedAllocations;
+	uint64 TotalAllocatedMemory = 0;
+	uint64 NumAllocations = 0;
+	GenerateScopedAllocations( AllocationMap, SnapshotScopedAllocations, TotalAllocatedMemory, NumAllocations );
+	SnapshotsWithScopedAllocations.Add( SnapshotName, SnapshotScopedAllocations );
+
+	// Decode callstacks.
+	// Replace encoded callstacks with human readable name. For easier debugging.
+	TMap<FString, FSizeAndCount> SnapshotDecodedScopedAllocations;
+	for (auto& It : SnapshotScopedAllocations)
+	{
+		const FString DecodedCallstack = GetCallstack( It.Key );
+		SnapshotDecodedScopedAllocations.Add( DecodedCallstack, It.Value );
+	}
+	SnapshotsWithDecodedScopedAllocations.Add( SnapshotName, SnapshotDecodedScopedAllocations );
+
+	UE_LOG( LogStats, Warning, TEXT( "PrepareSnapshot: %s Num: %i/%i Total: %.2f MB / %llu" ), *SnapshotName.GetPlainNameString(), AllocationMap.Num(), SnapshotScopedAllocations.Num(), TotalAllocatedMemory / 1024.0f / 1024.0f, NumAllocations );
+}
+
+void FStatsMemoryDumpCommand::CompareSnapshots_FName( const FName BeginSnaphotName, const FName EndSnaphotName, TMap<FName, FSizeAndCount>& out_Result )
+{
+	const auto BeginSnaphotPtr = SnapshotsWithScopedAllocations.Find( BeginSnaphotName );
+	const auto EndSnapshotPtr = SnapshotsWithScopedAllocations.Find( EndSnaphotName );
+	if (BeginSnaphotPtr && EndSnapshotPtr)
+	{
+		// Process data.
+		TMap<FName, FSizeAndCount> BeginSnaphot = *BeginSnaphotPtr;
+		TMap<FName, FSizeAndCount> EndSnaphot = *EndSnapshotPtr;
+		TMap<FName, FSizeAndCount> Result;
+
+		for (const auto& It : EndSnaphot)
+		{
+			const FName Callstack = It.Key;
+			const FSizeAndCount EndScopedAlloc = It.Value;
+
+			const FSizeAndCount* BeginScopedAllocPtr = BeginSnaphot.Find( Callstack );
+			if (BeginSnaphotPtr)
+			{
+				FSizeAndCount SizeAndCount;
+				SizeAndCount += EndScopedAlloc;
+				SizeAndCount -= *BeginScopedAllocPtr;
+
+				if (SizeAndCount.IsAlive())
+				{
+					Result.Add( Callstack, SizeAndCount );
+				}
+			}
+		}
+
+		// Sort by size.
+		out_Result.ValueSort( FSizeAndCountGreater() );
+	}
+}
+
+void FStatsMemoryDumpCommand::CompareSnapshots_FString( const FName BeginSnaphotName, const FName EndSnaphotName, TMap<FString, FSizeAndCount>& out_Result )
+{
+	const auto BeginSnaphotPtr = SnapshotsWithDecodedScopedAllocations.Find( BeginSnaphotName );
+	const auto EndSnapshotPtr = SnapshotsWithDecodedScopedAllocations.Find( EndSnaphotName );
+	if (BeginSnaphotPtr && EndSnapshotPtr)
+	{
+		// Process data.
+		TMap<FString, FSizeAndCount> BeginSnaphot = *BeginSnaphotPtr;
+		TMap<FString, FSizeAndCount> EndSnaphot = *EndSnapshotPtr;
+
+		for (const auto& It : EndSnaphot)
+		{
+			const FString& Callstack = It.Key;
+			const FSizeAndCount EndScopedAlloc = It.Value;
+
+			const FSizeAndCount* BeginScopedAllocPtr = BeginSnaphot.Find( Callstack );
+			if (BeginScopedAllocPtr)
+			{
+				FSizeAndCount SizeAndCount;
+				SizeAndCount += EndScopedAlloc;
+				SizeAndCount -= *BeginScopedAllocPtr;
+
+				if (SizeAndCount.IsAlive())
+				{
+					out_Result.Add( Callstack, SizeAndCount );
+				}
+			}
+			else
+			{
+				out_Result.Add( Callstack, EndScopedAlloc );
+			}
+		}
+
+		// Sort by size.
+		out_Result.ValueSort( FSizeAndCountGreater() );
+	}
+}
+
diff --git a/Engine/Source/Programs/UnrealFrontend/Private/Commands/StatsDumpMemoryCommand.h b/Engine/Source/Programs/UnrealFrontend/Private/Commands/StatsDumpMemoryCommand.h
--- a/Engine/Source/Programs/UnrealFrontend/Private/Commands/StatsDumpMemoryCommand.h
+++ b/Engine/Source/Programs/UnrealFrontend/Private/Commands/StatsDumpMemoryCommand.h
@@ -5,7 +5,89 @@
 #include "StatsData.h"
 #include "StatsFile.h"
 
-struct FAllocationInfo;
+/** Simple allocation info. Assumes the sequence tag are unique and doesn't wrap. */
+struct FAllocationInfo
+{
+	uint64 OldPtr;
+	uint64 Ptr;
+	int64 Size;
+	FName EncodedCallstack;
+	uint32 SequenceTag;
+	EMemoryOperation Op;
+	bool bHasBrokenCallstack;
+
+	FAllocationInfo(
+		uint64 InOldPtr,
+		uint64 InPtr,
+		int64 InSize,
+		const TArray<FName>& InCallstack,
+		uint32 InSequenceTag,
+		EMemoryOperation InOp,
+		bool bInHasBrokenCallstack
+		);
+
+	FAllocationInfo( const FAllocationInfo& Other );
+
+	bool operator<(const FAllocationInfo& Other) const
+	{
+		return SequenceTag < Other.SequenceTag;
+	}
+
+	bool operator==(const FAllocationInfo& Other) const
+	{
+		return SequenceTag == Other.SequenceTag;
+	}
+};
+
+struct FSizeAndCount
+{
+	int64 Size;
+	int64 Count;
+	double SizeMB;
+
+	FSizeAndCount()
+		: Size( 0 )
+		, Count( 0 )
+		, SizeMB( 0 )
+	{}
+
+	FSizeAndCount& operator+=(const FSizeAndCount& Other)
+	{
+		if ((UPTRINT*)this != (UPTRINT*)&Other)
+		{
+			AddAllocation( Other.Size, Other.Count );
+		}
+		return *this;
+	}
+
+	FSizeAndCount& operator-=(const FSizeAndCount& Other)
+	{
+		if ((UPTRINT*)this != (UPTRINT*)&Other)
+		{
+			RemoveAllocation( Other.Size, Other.Count );
+		}
+		return *this;
+	}
+
+	void AddAllocation( int64 SizeToAdd, int64 CountToAdd = 1 )
+	{
+		Size += SizeToAdd;
+		Count += CountToAdd;
+
+		const double InvMB = 1.0 / 1024.0 / 1024.0;
+		SizeMB = Size * InvMB;
+	}
+
+	void RemoveAllocation( int64 SizeToRemove, int64 CountToRemove )
+	{
+		AddAllocation( -SizeToRemove, -CountToRemove );
+	}
+
+	bool IsAlive()
+	{
+		return Size > 0;
+	}
+};
 
 /** An example command loading a raw stats file and dumping memory usage. */
 class FStatsMemoryDumpCommand
@@ -30,18 +112,43 @@ protected:
 	void GenerateMemoryUsageReport( const TMap<uint64, FAllocationInfo>& AllocationMap );
 
 	/** Generates scoped allocation statistics. */
-	void ProcessingScopedAllocations( const TMap<uint64, FAllocationInfo>& AllocationMap );
+	void ProcessAndDumpScopedAllocations( const TMap<uint64, FAllocationInfo>& AllocationMap );
 
 	/** Generates UObject allocation statistics. */
-	void ProcessingUObjectAllocations( const TMap<uint64, FAllocationInfo>& AllocationMap );
+	void ProcessAndDumpUObjectAllocations( const TMap<uint64, FAllocationInfo>& AllocationMap );
+
+	void DumpScopedAllocations( const TCHAR* Name, const TMap<FString, FSizeAndCount>& ScopedAllocations );
+
+	/** Generates callstack based allocation map. */
+	void GenerateScopedAllocations( const TMap<uint64, FAllocationInfo>& AllocationMap, TMap<FName, FSizeAndCount>& out_ScopedAllocations, uint64& TotalAllocatedMemory, uint64& NumAllocations );
+
+	/** Prepares data for a snapshot. */
+	void PrepareSnapshot( const FName SnapshotName, const TMap<uint64, FAllocationInfo>& AllocationMap );
+
+	void CompareSnapshots_FName( const FName BeginSnaphotName, const FName EndSnaphotName, TMap<FName, FSizeAndCount>& out_Result );
+	void CompareSnapshots_FString( const FName BeginSnaphotName, const FName EndSnaphotName, TMap<FString, FSizeAndCount>& out_Result );
 
 	/** Stats thread state, mostly used to manage the stats metadata. */
 	FStatsThreadState StatsThreadStats;
 	FStatsReadStream Stream;
 
-	/** All names that containa a path to an UObject. */
+	/** All names that contains a path to an UObject. */
 	TSet<FName> UObjectNames;
 
+	/** The sequence tag mapping to the named markers. */
+	TArray<TPair<uint32, FName>> Snapshots;
+
+	/** The sequence tag mapping to the named markers that need to processed. */
+	TArray<TPair<uint32, FName>> SnapshotsToBeProcessed;
+
+	/** Snapshots with allocation map. */
+	TMap<FName, TMap<uint64, FAllocationInfo> > SnapshotsWithAllocationMap;
+
+	/** Snapshots with callstack based allocation map. */
+	TMap<FName, TMap<FName, FSizeAndCount> > SnapshotsWithScopedAllocations;
+
+	TMap<FName, TMap<FString, FSizeAndCount> > SnapshotsWithDecodedScopedAllocations;
+
 	/** Filepath to the raw stats file. */
 	FString SourceFilepath;
 };
