Profiler - Fixed CompareSnapshots_FName

diff --git a/Engine/Source/Programs/UnrealFrontend/Private/Commands/StatsDumpMemoryCommand.cpp b/Engine/Source/Programs/UnrealFrontend/Private/Commands/StatsDumpMemoryCommand.cpp
--- a/Engine/Source/Programs/UnrealFrontend/Private/Commands/StatsDumpMemoryCommand.cpp
+++ b/Engine/Source/Programs/UnrealFrontend/Private/Commands/StatsDumpMemoryCommand.cpp
@@ -1133,7 +1133,7 @@ void FStatsMemoryDumpCommand::CompareSnapshots_FName( const FName BeginSnaphotNa
 			const FSizeAndCount EndScopedAlloc = It.Value;
 
 			const FSizeAndCount* BeginScopedAllocPtr = BeginSnaphot.Find( Callstack );
-			if (BeginSnaphotPtr)
+			if (BeginScopedAllocPtr)
 			{
 				FSizeAndCount SizeAndCount;
 				SizeAndCount += EndScopedAlloc;
@@ -1141,9 +1141,13 @@ void FStatsMemoryDumpCommand::CompareSnapshots_FName( const FName BeginSnaphotNa
 
 				if (SizeAndCount.IsAlive())
 				{
-					Result.Add( Callstack, SizeAndCount );
+					out_Result.Add( Callstack, SizeAndCount );
 				}
 			}
+			else
+			{
+				out_Result.Add( Callstack, EndScopedAlloc );
+			}
 		}
 
 		// Sort by size.
