Profiler - Reverted decoded callstack, better for displaying in the external app, added first version of generating scoped tree allocations

diff --git a/Engine/Source/Programs/UnrealFrontend/Private/Commands/StatsDumpMemoryCommand.cpp b/Engine/Source/Programs/UnrealFrontend/Private/Commands/StatsDumpMemoryCommand.cpp
--- a/Engine/Source/Programs/UnrealFrontend/Private/Commands/StatsDumpMemoryCommand.cpp
+++ b/Engine/Source/Programs/UnrealFrontend/Private/Commands/StatsDumpMemoryCommand.cpp
@@ -8,61 +8,87 @@
 	Callstack decoding/encoding
 -----------------------------------------------------------------------------*/
 
-static const TCHAR* CallstackSeparator = TEXT( "+" );
-
-static FString EncodeCallstack( const TArray<FName>& Callstack )
+/** Helper struct used to manipulate stats based callstacks. */
+struct FStatsCallstack
 {
-	FString Result;
-	for (const auto& Name : Callstack)
+	/** Separator. */
+	static const TCHAR* CallstackSeparator;
+
+	static FString Encode( const TArray<FName>& Callstack )
 	{
-		Result += TTypeToString<int32>::ToString( (int32)Name.GetComparisonIndex() );
-		Result += CallstackSeparator;
+		FString Result;
+		for (const auto& Name : Callstack)
+		{
+			Result += TTypeToString<int32>::ToString( (int32)Name.GetComparisonIndex() );
+			Result += CallstackSeparator;
+		}
+		return Result;
 	}
-	return Result;
-}
 
-static void DecodeCallstack( const FName& EncodedCallstack, TArray<FString>& out_DecodedCallstack )
-{
-	EncodedCallstack.ToString().ParseIntoArray( out_DecodedCallstack, CallstackSeparator, true );
-}
+	static void DecodeToStrings( const FName& EncodedCallstack, TArray<FString>& out_DecodedCallstack )
+	{
+		EncodedCallstack.ToString().ParseIntoArray( out_DecodedCallstack, CallstackSeparator, true );
+	}
 
-static FString GetCallstack( const FName& EncodedCallstack )
-{
-	FString Result;
+	static void DecodeToNames( const FName& EncodedCallstack, TArray<FName>& out_DecodedCallstack )
+	{
+		TArray<FString> DecodedCallstack;
+		DecodeToStrings( EncodedCallstack, DecodedCallstack );
+
+		// Convert back to FNames
+		for (const auto& It : DecodedCallstack)
+		{
+			NAME_INDEX NameIndex = 0;
+			TTypeFromString<NAME_INDEX>::FromString( NameIndex, *It );
+			const FName LongName = FName( NameIndex, NameIndex, 0 );
 
-	TArray<FString> DecodedCallstack;
-	DecodeCallstack( EncodedCallstack, DecodedCallstack );
+			out_DecodedCallstack.Add( LongName );
+		}
+	}
 
-	for (int32 Index = DecodedCallstack.Num() - 1; Index >= 0; --Index)
+	static FString GetHumanReadable( const FName& EncodedCallstack )
 	{
-		NAME_INDEX NameIndex = 0;
-		TTypeFromString<NAME_INDEX>::FromString( NameIndex, *DecodedCallstack[Index] );
-
-		const FName LongName = FName( NameIndex, NameIndex, 0 );
+		TArray<FName> DecodedCallstack;
+		DecodeToNames( EncodedCallstack, DecodedCallstack );
+		const FString Result = GetHumanReadable( DecodedCallstack );
+		return Result;
+	}
 
-		const FString ShortName = FStatNameAndInfo::GetShortNameFrom( LongName ).ToString();
-		//const FString Group = FStatNameAndInfo::GetGroupNameFrom( LongName ).ToString();
-		FString Desc = FStatNameAndInfo::GetDescriptionFrom( LongName );
-		Desc.Trim();
+	static FString GetHumanReadable( const TArray<FName>& DecodedCallstack )
+	{
+		FString Result;
 
-		if (Desc.Len() == 0)
+		const int32 NumEntries = DecodedCallstack.Num();
+		//for (int32 Index = DecodedCallstack.Num() - 1; Index >= 0; --Index)
+		for (int32 Index = 0; Index < NumEntries; ++Index)
 		{
-			Result += ShortName;
-		}
-		else
-		{
-			Result += Desc;
-		}
+			const FName LongName = DecodedCallstack[Index];
+			const FString ShortName = FStatNameAndInfo::GetShortNameFrom( LongName ).ToString();
+			//const FString Group = FStatNameAndInfo::GetGroupNameFrom( LongName ).ToString();
+			FString Desc = FStatNameAndInfo::GetDescriptionFrom( LongName );
+			Desc.Trim();
 
-		if (Index > 0)
-		{
-			Result += TEXT( " <- " );
+			if (Desc.Len() == 0)
+			{
+				Result += ShortName;
+			}
+			else
+			{
+				Result += Desc;
+			}
+
+			if (Index != NumEntries - 1)
+			{
+				Result += TEXT( " -> " );
+			}
 		}
+
+		Result.ReplaceInline( TEXT( "STAT_" ), TEXT( "" ), ESearchCase::CaseSensitive );
+		return Result;
 	}
+};
 
-	Result.ReplaceInline( TEXT( "STAT_" ), TEXT( "" ), ESearchCase::CaseSensitive );
-	return Result;
-}
+const TCHAR* FStatsCallstack::CallstackSeparator = TEXT( "+" );
 
 /*-----------------------------------------------------------------------------
 	Allocation info
@@ -72,7 +98,7 @@ FAllocationInfo::FAllocationInfo( uint64 InOldPtr, uint64 InPtr, int64 InSize, c
 	: OldPtr( InOldPtr )
 	, Ptr( InPtr )
 	, Size( InSize )
-	, EncodedCallstack( *EncodeCallstack( InCallstack ) )
+	, EncodedCallstack( *FStatsCallstack::Encode( InCallstack ) )
 	, SequenceTag( InSequenceTag )
 	, Op( InOp )
 	, bHasBrokenCallstack( bInHasBrokenCallstack )
@@ -107,6 +133,14 @@ struct FSizeAndCountGreater
 	}
 };
 
+struct FNodeAllocationInfoGreater
+{
+	FORCEINLINE bool operator()( const FNodeAllocationInfo& A, const FNodeAllocationInfo& B ) const
+	{
+		return B.Size < A.Size;
+	}
+};
+
 /*-----------------------------------------------------------------------------
 	Stats stack helpers
 -----------------------------------------------------------------------------*/
@@ -303,6 +337,14 @@ void FStatsMemoryDumpCommand::InternalRun()
 }
 
 
+void GetCallstackFromState( FStackState* StackState, TArray<FName>& out_StatsBasedCallstack )
+{
+	for (const auto& StackName : StackState->Stack)
+	{
+		out_StatsBasedCallstack.Add( StackName );
+	}
+}
+
 void FStatsMemoryDumpCommand::ProcessMemoryOperations( const TMap<int64, FStatPacketArray>& CombinedHistory )
 {
 	// This is only example code, no fully implemented, may sometimes crash.
@@ -430,10 +472,8 @@ void FStatsMemoryDumpCommand::ProcessMemoryOperations( const TMap<int64, FStatPa
 
 							// Create a callstack.
 							TArray<FName> StatsBasedCallstack;
-							for (const auto& StackName : StackState->Stack)
-							{
-								StatsBasedCallstack.Add( StackName );
-							}
+							GetCallstackFromState( StackState, StatsBasedCallstack );
+
 
 							// Add a new allocation.
 							SequenceAllocationArray.Add(
@@ -475,10 +515,7 @@ void FStatsMemoryDumpCommand::ProcessMemoryOperations( const TMap<int64, FStatPa
 
 							// Create a callstack.
 							TArray<FName> StatsBasedCallstack;
-							for (const auto& StackName : StackState->Stack)
-							{
-								StatsBasedCallstack.Add( StackName );
-							}
+							GetCallstackFromState( StackState, StatsBasedCallstack );
 
 							// Add a new realloc.
 							SequenceAllocationArray.Add(
@@ -506,10 +543,7 @@ void FStatsMemoryDumpCommand::ProcessMemoryOperations( const TMap<int64, FStatPa
 
 							// Create a callstack.
 							TArray<FName> StatsBasedCallstack;
-							for (const auto& StackName : StackState->Stack)
-							{
-								StatsBasedCallstack.Add( StackName );
-							}
+							GetCallstackFromState( StackState, StatsBasedCallstack );
 
 							// Add a new free.
 							SequenceAllocationArray.Add(
@@ -522,7 +556,6 @@ void FStatsMemoryDumpCommand::ProcessMemoryOperations( const TMap<int64, FStatPa
 								EMemoryOperation::Free,
 								StackState->bIsBrokenCallstack
 								) );
-							//LastSequenceTagForNamedMarker = SequenceTag;????
 						}
 						else
 						{
@@ -681,19 +714,19 @@ void FStatsMemoryDumpCommand::ProcessMemoryOperations( const TMap<int64, FStatPa
 #if	UE_BUILD_DEBUG
 				if( FoundAndFreed )
 				{
-					const FString FoundAndFreedCallstack = GetCallstack( FoundAndFreed->EncodedCallstack );
+					const FString FoundAndFreedCallstack = FStatsCallstack::GetHumanReadable( FoundAndFreed->EncodedCallstack );
 				}
 
 				if( FoundAndAllocated )
 				{
-					const FString FoundAndAllocatedCallstack = GetCallstack( FoundAndAllocated->EncodedCallstack );
+					const FString FoundAndAllocatedCallstack = FStatsCallstack::GetHumanReadable( FoundAndAllocated->EncodedCallstack );
 				}
 
 				NumDuplicatedMemoryOperations++;
 
 
-				const FString FoundCallstack = GetCallstack( Found->EncodedCallstack );
-				const FString AllocCallstack = GetCallstack( Alloc.EncodedCallstack );
+				const FString FoundCallstack = FStatsCallstack::GetHumanReadable( Found->EncodedCallstack );
+				const FString AllocCallstack = FStatsCallstack::GetHumanReadable( Alloc.EncodedCallstack );
 				UE_LOG( LogStats, Warning, TEXT( "DuplicateAlloc: %s Size: %i/%i Ptr: %i/%i Tag: %i/%i" ), *AllocCallstack, Found->Size, Alloc.Size, Found->Ptr, Alloc.Ptr, Found->SequenceTag, Alloc.SequenceTag );
 #endif // UE_BUILD_DEBUG
 
@@ -733,7 +766,7 @@ void FStatsMemoryDumpCommand::ProcessMemoryOperations( const TMap<int64, FStatPa
 				// Or Realloc after Malloc(0)
 
 #if	UE_BUILD_DEBUG
-				const FString ReallocCallstack = GetCallstack( Alloc.EncodedCallstack );
+				const FString ReallocCallstack = FStatsCallstack::GetHumanReadable( Alloc.EncodedCallstack );
 				UE_LOG( LogStats, Warning, TEXT( "ReallocWithoutAlloc: %s %i %i/%i [%i]" ), *ReallocCallstack, Alloc.Size, Alloc.OldPtr, Alloc.Ptr, Alloc.SequenceTag );
 #endif // UE_BUILD_DEBUG
 
@@ -759,7 +792,7 @@ void FStatsMemoryDumpCommand::ProcessMemoryOperations( const TMap<int64, FStatPa
 				NumFWAMemoryOperations++;
 
 #if	UE_BUILD_DEBUG
-				const FString FWACallstack = GetCallstack( Alloc.EncodedCallstack );
+				const FString FWACallstack = FStatsCallstack::GetHumanReadable( Alloc.EncodedCallstack );
 				UE_LOG( LogStats, Warning, TEXT( "FreeWithoutAllocCallstack: %s %i" ), *FWACallstack, Alloc.Ptr );
 #endif // UE_BUILD_DEBUG
 			}
@@ -791,7 +824,7 @@ void FStatsMemoryDumpCommand::ProcessMemoryOperations( const TMap<int64, FStatPa
 	for( const auto& It : DuplicatedAllocMap )
 	{
 		const FAllocationInfo& Alloc = It.Value;
-		const FString AllocCallstack = GetCallstack( Alloc.EncodedCallstack );
+		const FString AllocCallstack = FStatsCallstack::GetHumanReadable( Alloc.EncodedCallstack );
 		UE_LOG( LogStats, Log, TEXT( "%lli (%.2f MB) %s" ), Alloc.Size, Alloc.Size / 1024.0f / 1024.0f, *AllocCallstack );
 
 		DisplayedSoFar += Alloc.Size;
@@ -806,7 +839,7 @@ void FStatsMemoryDumpCommand::ProcessMemoryOperations( const TMap<int64, FStatPa
 	// Frame-240 Frame-120 Frame-060
 	TMap<FString, FSizeAndCount> FrameBegin_Exit;
 	CompareSnapshots_FString( TEXT( "BeginSnapshot" ), TEXT( "EngineLoop.Exit" ), FrameBegin_Exit );
-	DumpScopedAllocations( TEXT( "FrameBegin_Exit" ), FrameBegin_Exit );
+	DumpScopedAllocations( TEXT( "Begin_Exit" ), FrameBegin_Exit );
 
 	TMap<FString, FSizeAndCount> Frame060_120;
 	CompareSnapshots_FString( TEXT( "Frame-060" ), TEXT( "Frame-120" ), Frame060_120 );
@@ -816,10 +849,104 @@ void FStatsMemoryDumpCommand::ProcessMemoryOperations( const TMap<int64, FStatPa
 	CompareSnapshots_FString( TEXT( "Frame-060" ), TEXT( "Frame-240" ), Frame060_240 );
 	DumpScopedAllocations( TEXT( "Frame060_240" ), Frame060_240 );
 
+	// Generate scoped tree view.
+	{
+		TMap<FName, FSizeAndCount> FrameBegin_Exit_FName;
+		CompareSnapshots_FName( TEXT( "BeginSnapshot" ), TEXT( "EngineLoop.Exit" ), FrameBegin_Exit_FName );
+
+		FNodeAllocationInfo Root;
+		Root.EncodedCallstack = TEXT( "ThreadRoot" );
+		Root.HumanReadableCallstack = TEXT( "ThreadRoot" );
+		GenerateScopedTreeAllocations( FrameBegin_Exit_FName, Root );
+	}
+
+
+	{
+		TMap<FName, FSizeAndCount> Frame060_240_FName;
+		CompareSnapshots_FName( TEXT( "Frame-060" ), TEXT( "Frame-240" ), Frame060_240_FName );
+
+		FNodeAllocationInfo Root;
+		Root.EncodedCallstack = TEXT( "ThreadRoot" );
+		Root.HumanReadableCallstack = TEXT( "ThreadRoot" );
+		GenerateScopedTreeAllocations( Frame060_240_FName, Root );
+	}
+
+
 	// For snapshot?
 	//GenerateMemoryUsageReport( AllocationMap );
 }
 
+
+void FNodeAllocationInfo::SortBySize()
+{
+	ChildNodes.ValueSort( FNodeAllocationInfoGreater() );
+	for (auto& It : ChildNodes)
+	{
+		It.Value->SortBySize();
+	}
+}
+
+
+void FStatsMemoryDumpCommand::GenerateScopedTreeAllocations( const TMap<FName, FSizeAndCount>& ScopedAllocations, FNodeAllocationInfo& out_Root )
+{
+	// Experimental code, partially optimized.
+	// Decode all scoped allocations.
+	// Generate tree for allocations and combine them.
+	
+	for (const auto& It : ScopedAllocations)
+	{
+		const FName& EncodedCallstack = It.Key;
+		const FSizeAndCount& SizeAndCount = It.Value;
+
+		// Decode callstack.
+		TArray<FName> DecodedCallstack;
+		FStatsCallstack::DecodeToNames( EncodedCallstack, DecodedCallstack );
+		
+		const int32 AllocationLenght = DecodedCallstack.Num();
+		check( DecodedCallstack.Num() > 0 );
+
+		FNodeAllocationInfo* CurrentNode = &out_Root;
+		// Accumulate with thread root node.
+		CurrentNode->Accumulate( SizeAndCount );
+
+		// Iterate through the callstack and prepare all nodes if needed, and accumulate memory.
+		TArray<FName> CurrentCallstack;
+		const int32 NumEntries = DecodedCallstack.Num();
+		for (int32 Idx1 = 0; Idx1 < NumEntries; ++Idx1)
+		{
+			const FName NodeName = DecodedCallstack[Idx1];
+			CurrentCallstack.Add( NodeName );
+
+			FNodeAllocationInfo* Node = nullptr;
+			const bool bContainsNode = CurrentNode->ChildNodes.Contains( NodeName );
+			if (!bContainsNode)
+			{
+				Node = new FNodeAllocationInfo;
+				Node->Depth = Idx1;
+				Node->DecodedCallstack = CurrentCallstack;
+				Node->EncodedCallstack = *FStatsCallstack::Encode( CurrentCallstack );
+				Node->HumanReadableCallstack = FStatsCallstack::GetHumanReadable( CurrentCallstack );
+
+				CurrentNode->ChildNodes.Add( NodeName, Node );
+			}
+			else
+			{
+				Node = CurrentNode->ChildNodes.FindChecked( NodeName );
+			}
+
+			// Accumulate memory usage and num allocations for all nodes in the callstack.
+			Node->Accumulate( SizeAndCount );
+			
+			// Move to the next node.
+			Node->Parent = CurrentNode;
+			CurrentNode = Node;
+		}
+	}
+
+	out_Root.SortBySize();
+}
+
+
 void FStatsMemoryDumpCommand::GenerateMemoryUsageReport( const TMap<uint64, FAllocationInfo>& AllocationMap )
 {
 	if( AllocationMap.Num() == 0 )
@@ -866,7 +993,7 @@ void FStatsMemoryDumpCommand::ProcessAndDumpScopedAllocations( const TMap<uint64
 		const FSizeAndCount& SizeAndCount = It.Value;
 		const FName& EncodedCallstack = It.Key;
 
-		const FString AllocCallstack = GetCallstack( EncodedCallstack );
+		const FString AllocCallstack = FStatsCallstack::GetHumanReadable( EncodedCallstack );
 
 		UE_LOG( LogStats, Log, TEXT( "%2i, %llu (%.2f MB), %llu, %s" ),
 				CurrentIndex,
@@ -936,7 +1063,7 @@ void FStatsMemoryDumpCommand::ProcessAndDumpUObjectAllocations( const TMap<uint6
 		if( UObjectClass == NAME_None )
 		{
 			TArray<FString> DecodedCallstack;
-			DecodeCallstack( Alloc.EncodedCallstack, DecodedCallstack );
+			FStatsCallstack::DecodeToStrings( Alloc.EncodedCallstack, DecodedCallstack );
 
 			for( int32 Index = DecodedCallstack.Num() - 1; Index >= 0; --Index )
 			{
@@ -1016,9 +1143,15 @@ void FStatsMemoryDumpCommand::ProcessAndDumpUObjectAllocations( const TMap<uint6
 
 void FStatsMemoryDumpCommand::DumpScopedAllocations( const TCHAR* Name, const TMap<FString, FSizeAndCount>& ScopedAllocations )
 {
+	if (ScopedAllocations.Num() == 0)
+	{
+		UE_LOG( LogStats, Warning, TEXT( "No scoped allocations: %s" ), Name );
+		return;
+	}
+
 	// This code is not optimized. 
 	FScopeLogTime SLT( TEXT( "ProcessingScopedAllocations" ), nullptr, FScopeLogTime::ScopeLog_Seconds );
-	UE_LOG( LogStats, Warning, TEXT( "Processing scoped allocations" ) );
+	UE_LOG( LogStats, Warning, TEXT( "Dumping scoped allocations: %s" ), Name );
 
 	const FString ReportName = FString::Printf( TEXT( "%s-Memory-Scoped-%s" ), *Stream.Header.PlatformName, Name );
 	FDiagnosticTableViewer MemoryReport( *FDiagnosticTableViewer::GetUniqueTemporaryFilePath( *ReportName ), true );
@@ -1108,12 +1241,12 @@ void FStatsMemoryDumpCommand::PrepareSnapshot( const FName SnapshotName, const T
 	TMap<FString, FSizeAndCount> SnapshotDecodedScopedAllocations;
 	for (auto& It : SnapshotScopedAllocations)
 	{
-		const FString DecodedCallstack = GetCallstack( It.Key );
-		SnapshotDecodedScopedAllocations.Add( DecodedCallstack, It.Value );
+		const FString HumanReadableCallstack = FStatsCallstack::GetHumanReadable( It.Key );
+		SnapshotDecodedScopedAllocations.Add( HumanReadableCallstack, It.Value );
 	}
 	SnapshotsWithDecodedScopedAllocations.Add( SnapshotName, SnapshotDecodedScopedAllocations );
 
-	UE_LOG( LogStats, Warning, TEXT( "PrepareSnapshot: %s Num: %i/%i Total: %.2f MB / %llu" ), *SnapshotName.GetPlainNameString(), AllocationMap.Num(), SnapshotScopedAllocations.Num(), TotalAllocatedMemory / 1024.0f / 1024.0f, NumAllocations );
+	UE_LOG( LogStats, Warning, TEXT( "PrepareSnapshot: %s Alloc: %i Scoped: %i Total: %.2f MB" ), *SnapshotName.GetPlainNameString(), AllocationMap.Num(), SnapshotScopedAllocations.Num(), TotalAllocatedMemory / 1024.0f / 1024.0f );
 }
 
 void FStatsMemoryDumpCommand::CompareSnapshots_FName( const FName BeginSnaphotName, const FName EndSnaphotName, TMap<FName, FSizeAndCount>& out_Result )
diff --git a/Engine/Source/Programs/UnrealFrontend/Private/Commands/StatsDumpMemoryCommand.h b/Engine/Source/Programs/UnrealFrontend/Private/Commands/StatsDumpMemoryCommand.h
--- a/Engine/Source/Programs/UnrealFrontend/Private/Commands/StatsDumpMemoryCommand.h
+++ b/Engine/Source/Programs/UnrealFrontend/Private/Commands/StatsDumpMemoryCommand.h
@@ -32,13 +32,9 @@ struct FAllocationInfo
 	{
 		return SequenceTag < Other.SequenceTag;
 	}
-
-	bool operator==(const FAllocationInfo& Other) const
-	{
-		return SequenceTag == Other.SequenceTag;
-	}
 };
 
+// #YRX_Profiler: Add TArray<FName> etc 2015-06-25 
 struct FSizeAndCount
 {
 	int64 Size;
@@ -89,6 +85,70 @@ struct FSizeAndCount
 	}
 };
 
+struct FNodeAllocationInfo
+{
+	/** Size as MBs for this node and children. */
+	double SizeMB;
+
+	/** Allocations count for this node and children. */
+	int64 Count;
+
+	/** Human readable callstack like 'GameThread -> Preinit -> LoadModule */
+	FString HumanReadableCallstack;
+
+	/** Encoded callstack like '45+656+6565'. */
+	FName EncodedCallstack;
+
+	/** Callstack indices, encoded callstack parsed into an array [45,656,6565] */
+	TArray<FName> DecodedCallstack;
+
+	/** Parent node. */
+	FNodeAllocationInfo* Parent;
+
+	/** Child nodes. */
+	TMap<FName, FNodeAllocationInfo*> ChildNodes;
+
+	/** Size for this node and children. */
+	int64 Size;
+
+	/** Node's depth. */
+	int32 Depth;
+
+	FNodeAllocationInfo()
+		: SizeMB( 0.0 )
+		, Count( 0 )
+		, Parent( nullptr )
+		, Size( 0 )
+		, Depth( 0 )
+	{}
+
+
+	~FNodeAllocationInfo()
+	{
+		DeleteAllChildrenNodes();
+	}
+
+	void Accumulate( const FSizeAndCount& Other )
+	{
+		Size += Other.Size;
+		Count += Other.Count;
+
+		const double InvMB = 1.0 / 1024.0 / 1024.0;
+		SizeMB = Size * InvMB;
+	}
+
+	void DeleteAllChildrenNodes()
+	{
+		for (auto& It : ChildNodes)
+		{
+			delete It.Value;
+		}
+		ChildNodes.Empty();
+	}
+
+	void SortBySize();
+};
+
 /** An example command loading a raw stats file and dumping memory usage. */
 class FStatsMemoryDumpCommand
 {
@@ -122,6 +182,8 @@ protected:
 	/** Generates callstack based allocation map. */
 	void GenerateScopedAllocations( const TMap<uint64, FAllocationInfo>& AllocationMap, TMap<FName, FSizeAndCount>& out_ScopedAllocations, uint64& TotalAllocatedMemory, uint64& NumAllocations );
 
+	void GenerateScopedTreeAllocations( const TMap<FName, FSizeAndCount>& ScopedAllocations, FNodeAllocationInfo& out_Root );
+
 	/** Prepares data for a snapshot. */
 	void PrepareSnapshot( const FName SnapshotName, const TMap<uint64, FAllocationInfo>& AllocationMap );
 
