Profiler - Moved some test code bo run only in the debug builds only

diff --git a/Engine/Source/Programs/UnrealFrontend/Private/Commands/StatsDumpMemoryCommand.cpp b/Engine/Source/Programs/UnrealFrontend/Private/Commands/StatsDumpMemoryCommand.cpp
--- a/Engine/Source/Programs/UnrealFrontend/Private/Commands/StatsDumpMemoryCommand.cpp
+++ b/Engine/Source/Programs/UnrealFrontend/Private/Commands/StatsDumpMemoryCommand.cpp
@@ -841,6 +841,7 @@ void FStatsMemoryDumpCommand::ProcessMemoryOperations( const TMap<int64, FStatPa
 	CompareSnapshots_FString( TEXT( "BeginSnapshot" ), TEXT( "EngineLoop.Exit" ), FrameBegin_Exit );
 	DumpScopedAllocations( TEXT( "Begin_Exit" ), FrameBegin_Exit );
 
+#if	UE_BUILD_DEBUG
 	TMap<FString, FSizeAndCount> Frame060_120;
 	CompareSnapshots_FString( TEXT( "Frame-060" ), TEXT( "Frame-120" ), Frame060_120 );
 	DumpScopedAllocations( TEXT( "Frame060_120" ), Frame060_120 );
@@ -870,10 +871,7 @@ void FStatsMemoryDumpCommand::ProcessMemoryOperations( const TMap<int64, FStatPa
 		Root.HumanReadableCallstack = TEXT( "ThreadRoot" );
 		GenerateScopedTreeAllocations( Frame060_240_FName, Root );
 	}
-
-
-	// For snapshot?
-	//GenerateMemoryUsageReport( AllocationMap );
+#endif // UE_BUILD_DEBUG
 }
 
 
