Profiler - Documentation pass, cleaned some code

diff --git a/Engine/Source/Programs/UnrealFrontend/Private/Commands/StatsDumpMemoryCommand.cpp b/Engine/Source/Programs/UnrealFrontend/Private/Commands/StatsDumpMemoryCommand.cpp
--- a/Engine/Source/Programs/UnrealFrontend/Private/Commands/StatsDumpMemoryCommand.cpp
+++ b/Engine/Source/Programs/UnrealFrontend/Private/Commands/StatsDumpMemoryCommand.cpp
@@ -4,6 +4,47 @@
 #include "StatsDumpMemoryCommand.h"
 #include "DiagnosticTable.h"
 
+/*-----------------------------------------------------------------------------
+	Sort helpers
+-----------------------------------------------------------------------------*/
+
+/** Sorts allocations by size. */
+struct FAllocationInfoSequenceTagLess
+{
+	FORCEINLINE bool operator()( const FAllocationInfo& A, const FAllocationInfo& B ) const
+	{
+		return A.SequenceTag < B.SequenceTag;
+	}
+};
+
+/** Sorts allocations by size. */
+struct FAllocationInfoSizeGreater
+{
+	FORCEINLINE bool operator()( const FAllocationInfo& A, const FAllocationInfo& B ) const
+	{
+		return B.Size < A.Size;
+	}
+};
+
+/** Sorts combined allocations by size. */
+struct FCombinedAllocationInfoSizeGreater
+{
+	FORCEINLINE bool operator()( const FCombinedAllocationInfo& A, const FCombinedAllocationInfo& B ) const
+	{
+		return B.Size < A.Size;
+	}
+};
+
+
+/** Sorts node allocations by size. */
+struct FNodeAllocationInfoSizeGreater
+{
+	FORCEINLINE bool operator()( const FNodeAllocationInfo& A, const FNodeAllocationInfo& B ) const
+	{
+		return B.Size < A.Size;
+	}
+};
+
 /*-----------------------------------------------------------------------------
 	Callstack decoding/encoding
 -----------------------------------------------------------------------------*/
@@ -14,6 +55,7 @@ struct FStatsCallstack
 	/** Separator. */
 	static const TCHAR* CallstackSeparator;
 
+	/** Encodes decoded callstack a string, to be like '45+656+6565'. */
 	static FString Encode( const TArray<FName>& Callstack )
 	{
 		FString Result;
@@ -25,11 +67,7 @@ struct FStatsCallstack
 		return Result;
 	}
 
-	static void DecodeToStrings( const FName& EncodedCallstack, TArray<FString>& out_DecodedCallstack )
-	{
-		EncodedCallstack.ToString().ParseIntoArray( out_DecodedCallstack, CallstackSeparator, true );
-	}
-
+	/** Decodes encoded callstack to an array of FNames. */
 	static void DecodeToNames( const FName& EncodedCallstack, TArray<FName>& out_DecodedCallstack )
 	{
 		TArray<FString> DecodedCallstack;
@@ -46,6 +84,7 @@ struct FStatsCallstack
 		}
 	}
 
+	/** Converts the encoded callstack into human readable callstack. */
 	static FString GetHumanReadable( const FName& EncodedCallstack )
 	{
 		TArray<FName> DecodedCallstack;
@@ -54,6 +93,7 @@ struct FStatsCallstack
 		return Result;
 	}
 
+	/** Converts the encoded callstack into human readable callstack. */
 	static FString GetHumanReadable( const TArray<FName>& DecodedCallstack )
 	{
 		FString Result;
@@ -86,6 +126,13 @@ struct FStatsCallstack
 		Result.ReplaceInline( TEXT( "STAT_" ), TEXT( "" ), ESearchCase::CaseSensitive );
 		return Result;
 	}
+
+protected:
+	/** Decodes encoded callstack to an array of strings. Where each string is the index of the FName. */
+	static void DecodeToStrings( const FName& EncodedCallstack, TArray<FString>& out_DecodedCallstack )
+	{
+		EncodedCallstack.ToString().ParseIntoArray( out_DecodedCallstack, CallstackSeparator, true );
+	}
 };
 
 const TCHAR* FStatsCallstack::CallstackSeparator = TEXT( "+" );
@@ -118,28 +165,26 @@ FAllocationInfo::FAllocationInfo( const FAllocationInfo& Other )
 
 }
 
-struct FAllocationInfoGreater
-{
-	FORCEINLINE bool operator()( const FAllocationInfo& A, const FAllocationInfo& B ) const
-	{
-		return B.Size < A.Size;
-	}
-};
-struct FSizeAndCountGreater
+/*-----------------------------------------------------------------------------
+	FNodeAllocationInfo
+-----------------------------------------------------------------------------*/
+
+void FNodeAllocationInfo::SortBySize()
 {
-	FORCEINLINE bool operator()( const FSizeAndCount& A, const FSizeAndCount& B ) const
+	ChildNodes.ValueSort( FNodeAllocationInfoSizeGreater() );
+	for (auto& It : ChildNodes)
 	{
-		return B.Size < A.Size;
+		It.Value->SortBySize();
 	}
-};
+}
 
-struct FNodeAllocationInfoGreater
+
+void FNodeAllocationInfo::PrepareCallstackData( const TArray<FName>& InDecodedCallstack )
 {
-	FORCEINLINE bool operator()( const FNodeAllocationInfo& A, const FNodeAllocationInfo& B ) const
-	{
-		return B.Size < A.Size;
-	}
-};
+	DecodedCallstack = InDecodedCallstack;
+	EncodedCallstack = *FStatsCallstack::Encode( DecodedCallstack );
+	HumanReadableCallstack = FStatsCallstack::GetHumanReadable( DecodedCallstack );
+}
 
 /*-----------------------------------------------------------------------------
 	Stats stack helpers
@@ -162,6 +207,8 @@ struct FStackState
 	bool bIsBrokenCallstack;
 };
 
+
+
 /*-----------------------------------------------------------------------------
 	FStatsMemoryDumpCommand
 -----------------------------------------------------------------------------*/
@@ -185,6 +232,8 @@ void FStatsMemoryDumpCommand::InternalRun()
 		return;
 	}
 
+	FStatsReadStream Stream;
+
 	if( !Stream.ReadHeader( *FileReader ) )
 	{
 		UE_LOG( LogStats, Error, TEXT( "Could not open, bad magic: %s" ), *SourceFilepath );
@@ -196,6 +245,8 @@ void FStatsMemoryDumpCommand::InternalRun()
 	const bool bIsFinalized = Stream.Header.IsFinalized();
 	check( bIsFinalized );
 	check( Stream.Header.Version >= EStatMagicWithHeader::VERSION_6 );
+
+	FStatsThreadState StatsThreadStats;
 	StatsThreadStats.MarkAsLoaded();
 
 	TArray<FStatMessage> Messages;
@@ -208,6 +259,8 @@ void FStatsMemoryDumpCommand::InternalRun()
 		Stream.ReadFNamesAndMetadataMessages( *FileReader, MetadataMessages );
 		StatsThreadStats.ProcessMetaDataOnly( MetadataMessages );
 
+		ThreadIdToName = StatsThreadStats.Threads;
+
 		// Find all UObject metadata messages.
 		for( const auto& Meta : MetadataMessages )
 		{
@@ -232,6 +285,8 @@ void FStatsMemoryDumpCommand::InternalRun()
 		const bool bHasCompressedData = Stream.Header.HasCompressedData();
 		check( bHasCompressedData );
 
+		PlatformName = Stream.Header.PlatformName;
+
 		TMap<int64, FStatPacketArray> CombinedHistory;
 		int64 TotalDataSize = 0;
 		int64 TotalStatMessagesNum = 0;
@@ -337,14 +392,6 @@ void FStatsMemoryDumpCommand::InternalRun()
 }
 
 
-void GetCallstackFromState( FStackState* StackState, TArray<FName>& out_StatsBasedCallstack )
-{
-	for (const auto& StackName : StackState->Stack)
-	{
-		out_StatsBasedCallstack.Add( StackName );
-	}
-}
-
 void FStatsMemoryDumpCommand::ProcessMemoryOperations( const TMap<int64, FStatPacketArray>& CombinedHistory )
 {
 	// This is only example code, no fully implemented, may sometimes crash.
@@ -403,7 +450,7 @@ void FStatsMemoryDumpCommand::ProcessMemoryOperations( const TMap<int64, FStatPa
 			}
 
 			const FStatPacket& StatPacket = *Frame.Packets[PacketIndex];
-			const FName& ThreadFName = StatsThreadStats.Threads.FindChecked( StatPacket.ThreadId );
+			const FName& ThreadFName = ThreadIdToName.FindChecked( StatPacket.ThreadId );
 
 			FStackState* StackState = StackStates.Find( ThreadFName );
 			if( !StackState )
@@ -470,18 +517,13 @@ void FStatsMemoryDumpCommand::ProcessMemoryOperations( const TMap<int64, FStatPa
 							//ThreadStats->AddMemoryMessage( GET_STATFNAME( STAT_Memory_AllocSize ), Size );
 							//ThreadStats->AddMemoryMessage( GET_STATFNAME( STAT_Memory_OperationSequenceTag ), (int64)SequenceTag );
 
-							// Create a callstack.
-							TArray<FName> StatsBasedCallstack;
-							GetCallstackFromState( StackState, StatsBasedCallstack );
-
-
 							// Add a new allocation.
 							SequenceAllocationArray.Add(
 								FAllocationInfo(
 								0,
 								Ptr,
 								AllocSize,
-								StatsBasedCallstack,
+								StackState->Stack,
 								SequenceTag,
 								EMemoryOperation::Alloc,
 								StackState->bIsBrokenCallstack
@@ -513,17 +555,13 @@ void FStatsMemoryDumpCommand::ProcessMemoryOperations( const TMap<int64, FStatPa
 							//ThreadStats->AddMemoryMessage( GET_STATFNAME( STAT_Memory_AllocSize ), NewSize );
 							//ThreadStats->AddMemoryMessage( GET_STATFNAME( STAT_Memory_OperationSequenceTag ), (int64)SequenceTag );
 
-							// Create a callstack.
-							TArray<FName> StatsBasedCallstack;
-							GetCallstackFromState( StackState, StatsBasedCallstack );
-
 							// Add a new realloc.
 							SequenceAllocationArray.Add(
 								FAllocationInfo(
 								OldPtr, 
 								NewPtr,
 								ReallocSize,
-								StatsBasedCallstack,
+								StackState->Stack,
 								SequenceTag,
 								EMemoryOperation::Realloc,
 								StackState->bIsBrokenCallstack
@@ -541,17 +579,13 @@ void FStatsMemoryDumpCommand::ProcessMemoryOperations( const TMap<int64, FStatPa
 							//ThreadStats->AddMemoryMessage( GET_STATFNAME( STAT_Memory_FreePtr ), (uint64)(UPTRINT)Ptr | (uint64)EMemoryOperation::Free );	// 16 bytes total				
 							//ThreadStats->AddMemoryMessage( GET_STATFNAME( STAT_Memory_OperationSequenceTag ), (int64)SequenceTag );
 
-							// Create a callstack.
-							TArray<FName> StatsBasedCallstack;
-							GetCallstackFromState( StackState, StatsBasedCallstack );
-
 							// Add a new free.
 							SequenceAllocationArray.Add(
 								FAllocationInfo(
 								0,
 								Ptr,		
 								0,
-								StatsBasedCallstack,					
+								StackState->Stack,
 								SequenceTag,
 								EMemoryOperation::Free,
 								StackState->bIsBrokenCallstack
@@ -645,7 +679,7 @@ void FStatsMemoryDumpCommand::ProcessMemoryOperations( const TMap<int64, FStatPa
 	*/
 
 	// Sort all memory operation by the sequence tag, iterate through all operation and generate memory usage.
-	SequenceAllocationArray.Sort( TLess<FAllocationInfo>() );
+	SequenceAllocationArray.Sort( FAllocationInfoSequenceTagLess() );
 
 	// Named markers/snapshots
 
@@ -691,8 +725,6 @@ void FStatsMemoryDumpCommand::ProcessMemoryOperations( const TMap<int64, FStatPa
 			CurrentSnapshot = SnapshotsToBeProcessed[0];
 		}
 
-		
-
 		if (Alloc.Op == EMemoryOperation::Alloc)
 		{
 			if (Alloc.Size == 0)
@@ -808,7 +840,7 @@ void FStatsMemoryDumpCommand::ProcessMemoryOperations( const TMap<int64, FStatPa
 	UE_LOG( LogStats, Warning, TEXT( "NumZeroAllocs:                 %i" ), NumZeroAllocs );
 
 	// Dump problematic allocations
-	DuplicatedAllocMap.ValueSort( FAllocationInfoGreater() );
+	DuplicatedAllocMap.ValueSort( FAllocationInfoSizeGreater() );
 	//FreeWithoutAllocMap
 
 	uint64 TotalDuplicatedMemory = 0;
@@ -837,22 +869,22 @@ void FStatsMemoryDumpCommand::ProcessMemoryOperations( const TMap<int64, FStatPa
 	}
 
 	// Frame-240 Frame-120 Frame-060
-	TMap<FString, FSizeAndCount> FrameBegin_Exit;
+	TMap<FString, FCombinedAllocationInfo> FrameBegin_Exit;
 	CompareSnapshots_FString( TEXT( "BeginSnapshot" ), TEXT( "EngineLoop.Exit" ), FrameBegin_Exit );
 	DumpScopedAllocations( TEXT( "Begin_Exit" ), FrameBegin_Exit );
 
 #if	UE_BUILD_DEBUG
-	TMap<FString, FSizeAndCount> Frame060_120;
+	TMap<FString, FCombinedAllocationInfo> Frame060_120;
 	CompareSnapshots_FString( TEXT( "Frame-060" ), TEXT( "Frame-120" ), Frame060_120 );
 	DumpScopedAllocations( TEXT( "Frame060_120" ), Frame060_120 );
 
-	TMap<FString, FSizeAndCount> Frame060_240;
+	TMap<FString, FCombinedAllocationInfo> Frame060_240;
 	CompareSnapshots_FString( TEXT( "Frame-060" ), TEXT( "Frame-240" ), Frame060_240 );
 	DumpScopedAllocations( TEXT( "Frame060_240" ), Frame060_240 );
 
 	// Generate scoped tree view.
 	{
-		TMap<FName, FSizeAndCount> FrameBegin_Exit_FName;
+		TMap<FName, FCombinedAllocationInfo> FrameBegin_Exit_FName;
 		CompareSnapshots_FName( TEXT( "BeginSnapshot" ), TEXT( "EngineLoop.Exit" ), FrameBegin_Exit_FName );
 
 		FNodeAllocationInfo Root;
@@ -863,7 +895,7 @@ void FStatsMemoryDumpCommand::ProcessMemoryOperations( const TMap<int64, FStatPa
 
 
 	{
-		TMap<FName, FSizeAndCount> Frame060_240_FName;
+		TMap<FName, FCombinedAllocationInfo> Frame060_240_FName;
 		CompareSnapshots_FName( TEXT( "Frame-060" ), TEXT( "Frame-240" ), Frame060_240_FName );
 
 		FNodeAllocationInfo Root;
@@ -874,18 +906,7 @@ void FStatsMemoryDumpCommand::ProcessMemoryOperations( const TMap<int64, FStatPa
 #endif // UE_BUILD_DEBUG
 }
 
-
-void FNodeAllocationInfo::SortBySize()
-{
-	ChildNodes.ValueSort( FNodeAllocationInfoGreater() );
-	for (auto& It : ChildNodes)
-	{
-		It.Value->SortBySize();
-	}
-}
-
-
-void FStatsMemoryDumpCommand::GenerateScopedTreeAllocations( const TMap<FName, FSizeAndCount>& ScopedAllocations, FNodeAllocationInfo& out_Root )
+void FStatsMemoryDumpCommand::GenerateScopedTreeAllocations( const TMap<FName, FCombinedAllocationInfo>& ScopedAllocations, FNodeAllocationInfo& out_Root )
 {
 	// Experimental code, partially optimized.
 	// Decode all scoped allocations.
@@ -894,7 +915,7 @@ void FStatsMemoryDumpCommand::GenerateScopedTreeAllocations( const TMap<FName, F
 	for (const auto& It : ScopedAllocations)
 	{
 		const FName& EncodedCallstack = It.Key;
-		const FSizeAndCount& SizeAndCount = It.Value;
+		const FCombinedAllocationInfo& CombinedAllocation = It.Value;
 
 		// Decode callstack.
 		TArray<FName> DecodedCallstack;
@@ -905,7 +926,7 @@ void FStatsMemoryDumpCommand::GenerateScopedTreeAllocations( const TMap<FName, F
 
 		FNodeAllocationInfo* CurrentNode = &out_Root;
 		// Accumulate with thread root node.
-		CurrentNode->Accumulate( SizeAndCount );
+		CurrentNode->Accumulate( CombinedAllocation );
 
 		// Iterate through the callstack and prepare all nodes if needed, and accumulate memory.
 		TArray<FName> CurrentCallstack;
@@ -921,9 +942,7 @@ void FStatsMemoryDumpCommand::GenerateScopedTreeAllocations( const TMap<FName, F
 			{
 				Node = new FNodeAllocationInfo;
 				Node->Depth = Idx1;
-				Node->DecodedCallstack = CurrentCallstack;
-				Node->EncodedCallstack = *FStatsCallstack::Encode( CurrentCallstack );
-				Node->HumanReadableCallstack = FStatsCallstack::GetHumanReadable( CurrentCallstack );
+				Node->PrepareCallstackData( CurrentCallstack );
 
 				CurrentNode->ChildNodes.Add( NodeName, Node );
 			}
@@ -933,7 +952,7 @@ void FStatsMemoryDumpCommand::GenerateScopedTreeAllocations( const TMap<FName, F
 			}
 
 			// Accumulate memory usage and num allocations for all nodes in the callstack.
-			Node->Accumulate( SizeAndCount );
+			Node->Accumulate( CombinedAllocation );
 			
 			// Move to the next node.
 			Node->Parent = CurrentNode;
@@ -964,7 +983,7 @@ void FStatsMemoryDumpCommand::ProcessAndDumpScopedAllocations( const TMap<uint64
 	FScopeLogTime SLT( TEXT( "ProcessingScopedAllocations" ), nullptr, FScopeLogTime::ScopeLog_Seconds );
 	UE_LOG( LogStats, Warning, TEXT( "Processing scoped allocations" ) );
 
-	const FString ReportName = FString::Printf( TEXT( "%s-Memory-Scoped" ), *Stream.Header.PlatformName );
+	const FString ReportName = FString::Printf( TEXT( "%s-Memory-Scoped" ), *GetPlatformName() );
 	FDiagnosticTableViewer MemoryReport( *FDiagnosticTableViewer::GetUniqueTemporaryFilePath( *ReportName ), true );
 
 	// Write a row of headings for the table's columns.
@@ -974,41 +993,41 @@ void FStatsMemoryDumpCommand::ProcessAndDumpScopedAllocations( const TMap<uint64
 	MemoryReport.AddColumn( TEXT( "Callstack" ) );
 	MemoryReport.CycleRow();
 
-	TMap<FName, FSizeAndCount> ScopedAllocations;
+	TMap<FName, FCombinedAllocationInfo> CombinedAllocations;
 	uint64 TotalAllocatedMemory = 0;
 	uint64 NumAllocations = 0;
-	GenerateScopedAllocations( AllocationMap, ScopedAllocations, TotalAllocatedMemory, NumAllocations );
+	GenerateScopedAllocations( AllocationMap, CombinedAllocations, TotalAllocatedMemory, NumAllocations );
 
 	// Dump memory to the log.
-	ScopedAllocations.ValueSort( FSizeAndCountGreater() );
+	CombinedAllocations.ValueSort( FCombinedAllocationInfoSizeGreater() );
 
 	const float MaxPctDisplayed = 0.90f;
 	int32 CurrentIndex = 0;
 	uint64 DisplayedSoFar = 0;
 	UE_LOG( LogStats, Warning, TEXT( "Index, Size (Size MB), Count, Stat desc" ) );
-	for( const auto& It : ScopedAllocations )
+	for( const auto& It : CombinedAllocations )
 	{
-		const FSizeAndCount& SizeAndCount = It.Value;
+		const FCombinedAllocationInfo& CombinedAllocation = It.Value;
 		const FName& EncodedCallstack = It.Key;
 
 		const FString AllocCallstack = FStatsCallstack::GetHumanReadable( EncodedCallstack );
 
 		UE_LOG( LogStats, Log, TEXT( "%2i, %llu (%.2f MB), %llu, %s" ),
 				CurrentIndex,
-				SizeAndCount.Size,
-				SizeAndCount.Size / 1024.0f / 1024.0f,
-				SizeAndCount.Count,
+				CombinedAllocation.Size,
+				CombinedAllocation.Size / 1024.0f / 1024.0f,
+				CombinedAllocation.Count,
 				*AllocCallstack );
 
 		// Dump stats
-		MemoryReport.AddColumn( TEXT( "%llu" ), SizeAndCount.Size );
-		MemoryReport.AddColumn( TEXT( "%.2f MB" ), SizeAndCount.Size / 1024.0f / 1024.0f );
-		MemoryReport.AddColumn( TEXT( "%llu" ), SizeAndCount.Count );
+		MemoryReport.AddColumn( TEXT( "%llu" ), CombinedAllocation.Size );
+		MemoryReport.AddColumn( TEXT( "%.2f MB" ), CombinedAllocation.Size / 1024.0f / 1024.0f );
+		MemoryReport.AddColumn( TEXT( "%llu" ), CombinedAllocation.Count );
 		MemoryReport.AddColumn( *AllocCallstack );
 		MemoryReport.CycleRow();
 
 		CurrentIndex++;
-		DisplayedSoFar += SizeAndCount.Size;
+		DisplayedSoFar += CombinedAllocation.Size;
 
 		const float CurrentPct = (float)DisplayedSoFar / (float)TotalAllocatedMemory;
 		if( CurrentPct > MaxPctDisplayed )
@@ -1036,7 +1055,7 @@ void FStatsMemoryDumpCommand::ProcessAndDumpUObjectAllocations( const TMap<uint6
 	FScopeLogTime SLT( TEXT( "ProcessingUObjectAllocations" ), nullptr, FScopeLogTime::ScopeLog_Seconds );
 	UE_LOG( LogStats, Warning, TEXT( "Processing UObject allocations" ) );
 
-	const FString ReportName = FString::Printf( TEXT("%s-Memory-UObject"), *Stream.Header.PlatformName );
+	const FString ReportName = FString::Printf( TEXT( "%s-Memory-UObject" ), *GetPlatformName() );
 	FDiagnosticTableViewer MemoryReport( *FDiagnosticTableViewer::GetUniqueTemporaryFilePath( *ReportName ), true );
 
 	// Write a row of headings for the table's columns.
@@ -1046,7 +1065,7 @@ void FStatsMemoryDumpCommand::ProcessAndDumpUObjectAllocations( const TMap<uint6
 	MemoryReport.AddColumn( TEXT( "UObject class" ) );
 	MemoryReport.CycleRow();
 
-	TMap<FName, FSizeAndCount> UObjectAllocations;
+	TMap<FName, FCombinedAllocationInfo> UObjectAllocations;
 
 	// To minimize number of calls to expensive DecodeCallstack.
 	TMap<FName,FName> UObjectCallstackToClassMapping;
@@ -1060,15 +1079,12 @@ void FStatsMemoryDumpCommand::ProcessAndDumpUObjectAllocations( const TMap<uint6
 		FName UObjectClass = UObjectCallstackToClassMapping.FindRef( Alloc.EncodedCallstack );
 		if( UObjectClass == NAME_None )
 		{
-			TArray<FString> DecodedCallstack;
-			FStatsCallstack::DecodeToStrings( Alloc.EncodedCallstack, DecodedCallstack );
+			TArray<FName> DecodedCallstack;
+			FStatsCallstack::DecodeToNames( Alloc.EncodedCallstack, DecodedCallstack );
 
 			for( int32 Index = DecodedCallstack.Num() - 1; Index >= 0; --Index )
 			{
-				NAME_INDEX NameIndex = 0;
-				TTypeFromString<NAME_INDEX>::FromString( NameIndex, *DecodedCallstack[Index] );
-
-				const FName LongName = FName( NameIndex, NameIndex, 0 );
+				const FName LongName = DecodedCallstack[Index];
 				const bool bValid = UObjectNames.Contains( LongName );
 				if( bValid )
 				{
@@ -1082,8 +1098,8 @@ void FStatsMemoryDumpCommand::ProcessAndDumpUObjectAllocations( const TMap<uint6
 		
 		if( UObjectClass != NAME_None )
 		{
-			FSizeAndCount& SizeAndCount = UObjectAllocations.FindOrAdd( UObjectClass );
-			SizeAndCount.AddAllocation( Alloc.Size );
+			FCombinedAllocationInfo& CombinedAllocation = UObjectAllocations.FindOrAdd( UObjectClass );
+			CombinedAllocation += Alloc;
 
 			TotalAllocatedMemory += Alloc.Size;
 			NumAllocations++;
@@ -1091,7 +1107,7 @@ void FStatsMemoryDumpCommand::ProcessAndDumpUObjectAllocations( const TMap<uint6
 	}
 
 	// Dump memory to the log.
-	UObjectAllocations.ValueSort( FSizeAndCountGreater() );
+	UObjectAllocations.ValueSort( FCombinedAllocationInfoSizeGreater() );
 
 	const float MaxPctDisplayed = 0.90f;
 	int32 CurrentIndex = 0;
@@ -1099,25 +1115,25 @@ void FStatsMemoryDumpCommand::ProcessAndDumpUObjectAllocations( const TMap<uint6
 	UE_LOG( LogStats, Warning, TEXT( "Index, Size (Size MB), Count, UObject class" ) );
 	for( const auto& It : UObjectAllocations )
 	{
-		const FSizeAndCount& SizeAndCount = It.Value;
+		const FCombinedAllocationInfo& CombinedAllocation = It.Value;
 		const FName& UObjectClass = It.Key;
 
 		UE_LOG( LogStats, Log, TEXT( "%2i, %llu (%.2f MB), %llu, %s" ),
 				CurrentIndex,
-				SizeAndCount.Size,
-				SizeAndCount.Size / 1024.0f / 1024.0f,
-				SizeAndCount.Count,
+				CombinedAllocation.Size,
+				CombinedAllocation.Size / 1024.0f / 1024.0f,
+				CombinedAllocation.Count,
 				*UObjectClass.GetPlainNameString() );
 
 		// Dump stats
-		MemoryReport.AddColumn( TEXT( "%llu" ), SizeAndCount.Size );
-		MemoryReport.AddColumn( TEXT( "%.2f MB" ), SizeAndCount.Size / 1024.0f / 1024.0f );
-		MemoryReport.AddColumn( TEXT( "%llu" ), SizeAndCount.Count );
+		MemoryReport.AddColumn( TEXT( "%llu" ), CombinedAllocation.Size );
+		MemoryReport.AddColumn( TEXT( "%.2f MB" ), CombinedAllocation.Size / 1024.0f / 1024.0f );
+		MemoryReport.AddColumn( TEXT( "%llu" ), CombinedAllocation.Count );
 		MemoryReport.AddColumn( *UObjectClass.GetPlainNameString() );
 		MemoryReport.CycleRow();
 
 		CurrentIndex++;
-		DisplayedSoFar += SizeAndCount.Size;
+		DisplayedSoFar += CombinedAllocation.Size;
 
 		const float CurrentPct = (float)DisplayedSoFar / (float)TotalAllocatedMemory;
 		if( CurrentPct > MaxPctDisplayed )
@@ -1139,9 +1155,9 @@ void FStatsMemoryDumpCommand::ProcessAndDumpUObjectAllocations( const TMap<uint6
 	MemoryReport.CycleRow();
 }
 
-void FStatsMemoryDumpCommand::DumpScopedAllocations( const TCHAR* Name, const TMap<FString, FSizeAndCount>& ScopedAllocations )
+void FStatsMemoryDumpCommand::DumpScopedAllocations( const TCHAR* Name, const TMap<FString, FCombinedAllocationInfo>& CombinedAllocations )
 {
-	if (ScopedAllocations.Num() == 0)
+	if (CombinedAllocations.Num() == 0)
 	{
 		UE_LOG( LogStats, Warning, TEXT( "No scoped allocations: %s" ), Name );
 		return;
@@ -1151,7 +1167,7 @@ void FStatsMemoryDumpCommand::DumpScopedAllocations( const TCHAR* Name, const TM
 	FScopeLogTime SLT( TEXT( "ProcessingScopedAllocations" ), nullptr, FScopeLogTime::ScopeLog_Seconds );
 	UE_LOG( LogStats, Warning, TEXT( "Dumping scoped allocations: %s" ), Name );
 
-	const FString ReportName = FString::Printf( TEXT( "%s-Memory-Scoped-%s" ), *Stream.Header.PlatformName, Name );
+	const FString ReportName = FString::Printf( TEXT( "%s-Memory-Scoped-%s" ), *GetPlatformName(), Name );
 	FDiagnosticTableViewer MemoryReport( *FDiagnosticTableViewer::GetUniqueTemporaryFilePath( *ReportName ), true );
 
 	// Write a row of headings for the table's columns.
@@ -1162,33 +1178,33 @@ void FStatsMemoryDumpCommand::DumpScopedAllocations( const TCHAR* Name, const TM
 	MemoryReport.CycleRow();
 
 	
-	FSizeAndCount Total;
+	FCombinedAllocationInfo Total;
 
 	const float MaxPctDisplayed = 0.90f;
 	int32 CurrentIndex = 0;
 	UE_LOG( LogStats, Warning, TEXT( "Index, Size (Size MB), Count, Stat desc" ) );
-	for (const auto& It : ScopedAllocations)
+	for (const auto& It : CombinedAllocations)
 	{
-		const FSizeAndCount& SizeAndCount = It.Value;
+		const FCombinedAllocationInfo& CombinedAllocation = It.Value;
 		//const FName& EncodedCallstack = It.Key;
 		const FString AllocCallstack = It.Key;// GetCallstack( EncodedCallstack );
 
 		UE_LOG( LogStats, Log, TEXT( "%2i, %llu (%.2f MB), %llu, %s" ),
 				CurrentIndex,
-				SizeAndCount.Size,
-				SizeAndCount.Size / 1024.0f / 1024.0f,
-				SizeAndCount.Count,
+				CombinedAllocation.Size,
+				CombinedAllocation.Size / 1024.0f / 1024.0f,
+				CombinedAllocation.Count,
 				*AllocCallstack );
 
 		// Dump stats
-		MemoryReport.AddColumn( TEXT( "%llu" ), SizeAndCount.Size );
-		MemoryReport.AddColumn( TEXT( "%.2f MB" ), SizeAndCount.Size / 1024.0f / 1024.0f );
-		MemoryReport.AddColumn( TEXT( "%llu" ), SizeAndCount.Count );
+		MemoryReport.AddColumn( TEXT( "%llu" ), CombinedAllocation.Size );
+		MemoryReport.AddColumn( TEXT( "%.2f MB" ), CombinedAllocation.Size / 1024.0f / 1024.0f );
+		MemoryReport.AddColumn( TEXT( "%llu" ), CombinedAllocation.Count );
 		MemoryReport.AddColumn( *AllocCallstack );
 		MemoryReport.CycleRow();
 
 		CurrentIndex++;
-		Total += SizeAndCount;
+		Total += CombinedAllocation;
 	}
 
 	UE_LOG( LogStats, Warning, TEXT( "Allocated memory: %llu bytes (%.2f MB)" ), Total.Size, Total.SizeMB );
@@ -1204,22 +1220,22 @@ void FStatsMemoryDumpCommand::DumpScopedAllocations( const TCHAR* Name, const TM
 	MemoryReport.CycleRow();
 }
 
-void FStatsMemoryDumpCommand::GenerateScopedAllocations( const TMap<uint64, FAllocationInfo>& AllocationMap, TMap<FName, FSizeAndCount>& out_ScopedAllocations, uint64& TotalAllocatedMemory, uint64& NumAllocations )
+void FStatsMemoryDumpCommand::GenerateScopedAllocations( const TMap<uint64, FAllocationInfo>& AllocationMap, TMap<FName, FCombinedAllocationInfo>& out_CombinedAllocations, uint64& TotalAllocatedMemory, uint64& NumAllocations )
 {
 	FScopeLogTime SLT( TEXT( "GenerateScopedAllocations" ), nullptr, FScopeLogTime::ScopeLog_Milliseconds );
 
 	for (const auto& It : AllocationMap)
 	{
 		const FAllocationInfo& Alloc = It.Value;
-		FSizeAndCount& SizeAndCount = out_ScopedAllocations.FindOrAdd( Alloc.EncodedCallstack );
-		SizeAndCount.AddAllocation( Alloc.Size );
+		FCombinedAllocationInfo& CombinedAllocation = out_CombinedAllocations.FindOrAdd( Alloc.EncodedCallstack );
+		CombinedAllocation += Alloc;
 
 		TotalAllocatedMemory += Alloc.Size;
 		NumAllocations++;
 	}
 
 	// Sort by size.
-	out_ScopedAllocations.ValueSort( FSizeAndCountGreater() );
+	out_CombinedAllocations.ValueSort( FCombinedAllocationInfoSizeGreater() );
 }
 
 void FStatsMemoryDumpCommand::PrepareSnapshot( const FName SnapshotName, const TMap<uint64, FAllocationInfo>& AllocationMap )
@@ -1228,99 +1244,99 @@ void FStatsMemoryDumpCommand::PrepareSnapshot( const FName SnapshotName, const T
 
 	SnapshotsWithAllocationMap.Add( SnapshotName, AllocationMap );
 
-	TMap<FName, FSizeAndCount> SnapshotScopedAllocations;
+	TMap<FName, FCombinedAllocationInfo> SnapshotCombinedAllocations;
 	uint64 TotalAllocatedMemory = 0;
 	uint64 NumAllocations = 0;
-	GenerateScopedAllocations( AllocationMap, SnapshotScopedAllocations, TotalAllocatedMemory, NumAllocations );
-	SnapshotsWithScopedAllocations.Add( SnapshotName, SnapshotScopedAllocations );
+	GenerateScopedAllocations( AllocationMap, SnapshotCombinedAllocations, TotalAllocatedMemory, NumAllocations );
+	SnapshotsWithScopedAllocations.Add( SnapshotName, SnapshotCombinedAllocations );
 
 	// Decode callstacks.
 	// Replace encoded callstacks with human readable name. For easier debugging.
-	TMap<FString, FSizeAndCount> SnapshotDecodedScopedAllocations;
-	for (auto& It : SnapshotScopedAllocations)
+	TMap<FString, FCombinedAllocationInfo> SnapshotDecodedCombinedAllocations;
+	for (auto& It : SnapshotCombinedAllocations)
 	{
 		const FString HumanReadableCallstack = FStatsCallstack::GetHumanReadable( It.Key );
-		SnapshotDecodedScopedAllocations.Add( HumanReadableCallstack, It.Value );
+		SnapshotDecodedCombinedAllocations.Add( HumanReadableCallstack, It.Value );
 	}
-	SnapshotsWithDecodedScopedAllocations.Add( SnapshotName, SnapshotDecodedScopedAllocations );
+	SnapshotsWithDecodedScopedAllocations.Add( SnapshotName, SnapshotDecodedCombinedAllocations );
 
-	UE_LOG( LogStats, Warning, TEXT( "PrepareSnapshot: %s Alloc: %i Scoped: %i Total: %.2f MB" ), *SnapshotName.GetPlainNameString(), AllocationMap.Num(), SnapshotScopedAllocations.Num(), TotalAllocatedMemory / 1024.0f / 1024.0f );
+	UE_LOG( LogStats, Warning, TEXT( "PrepareSnapshot: %s Alloc: %i Scoped: %i Total: %.2f MB" ), *SnapshotName.GetPlainNameString(), AllocationMap.Num(), SnapshotCombinedAllocations.Num(), TotalAllocatedMemory / 1024.0f / 1024.0f );
 }
 
-void FStatsMemoryDumpCommand::CompareSnapshots_FName( const FName BeginSnaphotName, const FName EndSnaphotName, TMap<FName, FSizeAndCount>& out_Result )
+void FStatsMemoryDumpCommand::CompareSnapshots_FName( const FName BeginSnaphotName, const FName EndSnaphotName, TMap<FName, FCombinedAllocationInfo>& out_Result )
 {
 	const auto BeginSnaphotPtr = SnapshotsWithScopedAllocations.Find( BeginSnaphotName );
 	const auto EndSnapshotPtr = SnapshotsWithScopedAllocations.Find( EndSnaphotName );
 	if (BeginSnaphotPtr && EndSnapshotPtr)
 	{
 		// Process data.
-		TMap<FName, FSizeAndCount> BeginSnaphot = *BeginSnaphotPtr;
-		TMap<FName, FSizeAndCount> EndSnaphot = *EndSnapshotPtr;
-		TMap<FName, FSizeAndCount> Result;
+		TMap<FName, FCombinedAllocationInfo> BeginSnaphot = *BeginSnaphotPtr;
+		TMap<FName, FCombinedAllocationInfo> EndSnaphot = *EndSnapshotPtr;
+		TMap<FName, FCombinedAllocationInfo> Result;
 
 		for (const auto& It : EndSnaphot)
 		{
 			const FName Callstack = It.Key;
-			const FSizeAndCount EndScopedAlloc = It.Value;
+			const FCombinedAllocationInfo EndCombinedAlloc = It.Value;
 
-			const FSizeAndCount* BeginScopedAllocPtr = BeginSnaphot.Find( Callstack );
-			if (BeginScopedAllocPtr)
+			const FCombinedAllocationInfo* BeginCombinedAllocPtr = BeginSnaphot.Find( Callstack );
+			if (BeginCombinedAllocPtr)
 			{
-				FSizeAndCount SizeAndCount;
-				SizeAndCount += EndScopedAlloc;
-				SizeAndCount -= *BeginScopedAllocPtr;
+				FCombinedAllocationInfo CombinedAllocation;
+				CombinedAllocation += EndCombinedAlloc;
+				CombinedAllocation -= *BeginCombinedAllocPtr;
 
-				if (SizeAndCount.IsAlive())
+				if (CombinedAllocation.IsAlive())
 				{
-					out_Result.Add( Callstack, SizeAndCount );
+					out_Result.Add( Callstack, CombinedAllocation );
 				}
 			}
 			else
 			{
-				out_Result.Add( Callstack, EndScopedAlloc );
+				out_Result.Add( Callstack, EndCombinedAlloc );
 			}
 		}
 
 		// Sort by size.
-		out_Result.ValueSort( FSizeAndCountGreater() );
+		out_Result.ValueSort( FCombinedAllocationInfoSizeGreater() );
 	}
 }
 
-void FStatsMemoryDumpCommand::CompareSnapshots_FString( const FName BeginSnaphotName, const FName EndSnaphotName, TMap<FString, FSizeAndCount>& out_Result )
+void FStatsMemoryDumpCommand::CompareSnapshots_FString( const FName BeginSnaphotName, const FName EndSnaphotName, TMap<FString, FCombinedAllocationInfo>& out_Result )
 {
 	const auto BeginSnaphotPtr = SnapshotsWithDecodedScopedAllocations.Find( BeginSnaphotName );
 	const auto EndSnapshotPtr = SnapshotsWithDecodedScopedAllocations.Find( EndSnaphotName );
 	if (BeginSnaphotPtr && EndSnapshotPtr)
 	{
 		// Process data.
-		TMap<FString, FSizeAndCount> BeginSnaphot = *BeginSnaphotPtr;
-		TMap<FString, FSizeAndCount> EndSnaphot = *EndSnapshotPtr;
+		TMap<FString, FCombinedAllocationInfo> BeginSnaphot = *BeginSnaphotPtr;
+		TMap<FString, FCombinedAllocationInfo> EndSnaphot = *EndSnapshotPtr;
 
 		for (const auto& It : EndSnaphot)
 		{
 			const FString& Callstack = It.Key;
-			const FSizeAndCount EndScopedAlloc = It.Value;
+			const FCombinedAllocationInfo EndCombinedAlloc = It.Value;
 
-			const FSizeAndCount* BeginScopedAllocPtr = BeginSnaphot.Find( Callstack );
-			if (BeginScopedAllocPtr)
+			const FCombinedAllocationInfo* BeginCombinedAllocPtr = BeginSnaphot.Find( Callstack );
+			if (BeginCombinedAllocPtr)
 			{
-				FSizeAndCount SizeAndCount;
-				SizeAndCount += EndScopedAlloc;
-				SizeAndCount -= *BeginScopedAllocPtr;
+				FCombinedAllocationInfo CombinedAllocation;
+				CombinedAllocation += EndCombinedAlloc;
+				CombinedAllocation -= *BeginCombinedAllocPtr;
 
-				if (SizeAndCount.IsAlive())
+				if (CombinedAllocation.IsAlive())
 				{
-					out_Result.Add( Callstack, SizeAndCount );
+					out_Result.Add( Callstack, CombinedAllocation );
 				}
 			}
 			else
 			{
-				out_Result.Add( Callstack, EndScopedAlloc );
+				out_Result.Add( Callstack, EndCombinedAlloc );
 			}
 		}
 
 		// Sort by size.
-		out_Result.ValueSort( FSizeAndCountGreater() );
+		out_Result.ValueSort( FCombinedAllocationInfoSizeGreater() );
 	}
 }
 
diff --git a/Engine/Source/Programs/UnrealFrontend/Private/Commands/StatsDumpMemoryCommand.h b/Engine/Source/Programs/UnrealFrontend/Private/Commands/StatsDumpMemoryCommand.h
--- a/Engine/Source/Programs/UnrealFrontend/Private/Commands/StatsDumpMemoryCommand.h
+++ b/Engine/Source/Programs/UnrealFrontend/Private/Commands/StatsDumpMemoryCommand.h
@@ -5,17 +5,42 @@
 #include "StatsData.h"
 #include "StatsFile.h"
 
+/*-----------------------------------------------------------------------------
+	Basic structures
+-----------------------------------------------------------------------------*/
+
 /** Simple allocation info. Assumes the sequence tag are unique and doesn't wrap. */
 struct FAllocationInfo
 {
+	/** Old pointer, for Realloc. */
 	uint64 OldPtr;
+
+	/** Pointer, for Alloc, Realloc and Free. */
 	uint64 Ptr;
+
+	/** Size of the allocation. */
 	int64 Size;
+	
+	/**
+	 * Encoded callstack as FNames indices separated with +
+	 * 123+45+56+76
+	 */
 	FName EncodedCallstack;
+	
+	/**
+	 * Index of the memory operation. 
+	 * Stats are handled with packets sent from different threads.
+	 * This is used to sort the memory operations to be in the sequence order.
+	 */
 	uint32 SequenceTag;
+	
+	/** Memory operation, can be Alloc, Free or Realloc. */
 	EMemoryOperation Op;
+	
+	/** If true, the callstack is broken. Missing Start/End scope. */
 	bool bHasBrokenCallstack;
 
+	/** Initialization constructor. */
 	FAllocationInfo(
 		uint64 InOldPtr,
 		uint64 InPtr,
@@ -26,28 +51,40 @@ struct FAllocationInfo
 		bool bInHasBrokenCallstack
 		);
 
+	/** Copy constructor. */
 	FAllocationInfo( const FAllocationInfo& Other );
-
-	bool operator<(const FAllocationInfo& Other) const
-	{
-		return SequenceTag < Other.SequenceTag;
-	}
 };
 
-// #YRX_Profiler: Add TArray<FName> etc 2015-06-25 
-struct FSizeAndCount
+/** Contains data about combined allocations for an unique scope. */
+struct FCombinedAllocationInfo
 {
-	int64 Size;
-	int64 Count;
+	/** Size as MB. */
 	double SizeMB;
 
-	FSizeAndCount()
-		: Size( 0 )
+	/** Allocations count. */
+	int64 Count;
+
+	/** Human readable callstack like 'GameThread -> Preinit -> LoadModule */
+	FString HumanReadableCallstack;
+
+	/** Encoded callstack like '45+656+6565'. */
+	FName EncodedCallstack;
+
+	/** Callstack indices, encoded callstack parsed into an array [45,656,6565] */
+	TArray<FName> DecodedCallstack;
+
+	/** Size for this node and children. */
+	int64 Size;
+
+	/** Default constructor. */
+	FCombinedAllocationInfo()
+		: SizeMB( 0 )
 		, Count( 0 )
-		, SizeMB( 0 )
+		, Size( 0 )
 	{}
 
-	FSizeAndCount& operator+=(const FSizeAndCount& Other)
+	/** Operator for adding another combined allocation. */
+	FCombinedAllocationInfo& operator+=(const FCombinedAllocationInfo& Other)
 	{
 		if ((UPTRINT*)this != (UPTRINT*)&Other)
 		{
@@ -56,15 +93,36 @@ struct FSizeAndCount
 		return *this;
 	}
 
-	FSizeAndCount& operator-=(const FSizeAndCount& Other)
+	/** Operator for subtracting another combined allocation. */
+	FCombinedAllocationInfo& operator-=(const FCombinedAllocationInfo& Other)
 	{
 		if ((UPTRINT*)this != (UPTRINT*)&Other)
 		{
-			RemoveAllocation( Other.Size, Other.Count );
+			AddAllocation( -Other.Size, -Other.Count );
 		}
 		return *this;
 	}
 
+	/** Operator for adding a new allocation. */
+	FCombinedAllocationInfo& operator+=(const FAllocationInfo& Other)
+	{
+		if ((UPTRINT*)this != (UPTRINT*)&Other)
+		{
+			AddAllocation( Other.Size, 1 );
+		}
+		return *this;
+	}
+	
+	/**
+	 * @return true, if the combined allocation is alive, the allocated data is larger than 0
+	 */
+	bool IsAlive()
+	{
+		return Size > 0;
+	}
+
+protected:
+	/** Adds a new allocation. */
 	void AddAllocation( int64 SizeToAdd, int64 CountToAdd = 1 )
 	{
 		Size += SizeToAdd;
@@ -73,18 +131,9 @@ struct FSizeAndCount
 		const double InvMB = 1.0 / 1024.0 / 1024.0;
 		SizeMB = Size * InvMB;
 	}
-
-	void RemoveAllocation( int64 SizeToRemove, int64 CountToRemove )
-	{
-		AddAllocation( -SizeToRemove, -CountToRemove );
-	}
-
-	bool IsAlive()
-	{
-		return Size > 0;
-	}
 };
 
+/** Contains data about node allocations, parent and children allocations. */
 struct FNodeAllocationInfo
 {
 	/** Size as MBs for this node and children. */
@@ -114,6 +163,7 @@ struct FNodeAllocationInfo
 	/** Node's depth. */
 	int32 Depth;
 
+	/** Default constructor. */
 	FNodeAllocationInfo()
 		: SizeMB( 0.0 )
 		, Count( 0 )
@@ -122,13 +172,14 @@ struct FNodeAllocationInfo
 		, Depth( 0 )
 	{}
 
-
+	/** Destructor. */
 	~FNodeAllocationInfo()
 	{
 		DeleteAllChildrenNodes();
 	}
 
-	void Accumulate( const FSizeAndCount& Other )
+	/** Accumulates this node with the other combined allocation. */
+	void Accumulate( const FCombinedAllocationInfo& Other )
 	{
 		Size += Other.Size;
 		Count += Other.Count;
@@ -137,6 +188,14 @@ struct FNodeAllocationInfo
 		SizeMB = Size * InvMB;
 	}
 
+	/** Recursively sorts all children by size. */
+	void SortBySize();
+
+	/** Prepares all callstack data for this node. */
+	void PrepareCallstackData( const TArray<FName>& InDecodedCallstack );
+
+protected:
+	/** Recursively deletes all children. */
 	void DeleteAllChildrenNodes()
 	{
 		for (auto& It : ChildNodes)
@@ -145,8 +204,6 @@ struct FNodeAllocationInfo
 		}
 		ChildNodes.Empty();
 	}
-
-	void SortBySize();
 };
 
 /** An example command loading a raw stats file and dumping memory usage. */
@@ -177,22 +234,30 @@ protected:
 	/** Generates UObject allocation statistics. */
 	void ProcessAndDumpUObjectAllocations( const TMap<uint64, FAllocationInfo>& AllocationMap );
 
-	void DumpScopedAllocations( const TCHAR* Name, const TMap<FString, FSizeAndCount>& ScopedAllocations );
+	void DumpScopedAllocations( const TCHAR* Name, const TMap<FString, FCombinedAllocationInfo>& CombinedAllocations );
 
 	/** Generates callstack based allocation map. */
-	void GenerateScopedAllocations( const TMap<uint64, FAllocationInfo>& AllocationMap, TMap<FName, FSizeAndCount>& out_ScopedAllocations, uint64& TotalAllocatedMemory, uint64& NumAllocations );
+	void GenerateScopedAllocations( const TMap<uint64, FAllocationInfo>& AllocationMap, TMap<FName, FCombinedAllocationInfo>& out_CombinedAllocations, uint64& TotalAllocatedMemory, uint64& NumAllocations );
 
-	void GenerateScopedTreeAllocations( const TMap<FName, FSizeAndCount>& ScopedAllocations, FNodeAllocationInfo& out_Root );
+	void GenerateScopedTreeAllocations( const TMap<FName, FCombinedAllocationInfo>& ScopedAllocations, FNodeAllocationInfo& out_Root );
 
 	/** Prepares data for a snapshot. */
 	void PrepareSnapshot( const FName SnapshotName, const TMap<uint64, FAllocationInfo>& AllocationMap );
 
-	void CompareSnapshots_FName( const FName BeginSnaphotName, const FName EndSnaphotName, TMap<FName, FSizeAndCount>& out_Result );
-	void CompareSnapshots_FString( const FName BeginSnaphotName, const FName EndSnaphotName, TMap<FString, FSizeAndCount>& out_Result );
+	void CompareSnapshots_FName( const FName BeginSnaphotName, const FName EndSnaphotName, TMap<FName, FCombinedAllocationInfo>& out_Result );
+	void CompareSnapshots_FString( const FName BeginSnaphotName, const FName EndSnaphotName, TMap<FString, FCombinedAllocationInfo>& out_Result );
+
+
+	/**
+	 * @return Platform's name based on the loaded ue4statsraw file
+	 */
+	const FString& GetPlatformName() const
+	{
+		return PlatformName;
+	}
 
-	/** Stats thread state, mostly used to manage the stats metadata. */
-	FStatsThreadState StatsThreadStats;
-	FStatsReadStream Stream;
+	/** Map from a thread id to the thread name. */
+	TMap<uint32, FName> ThreadIdToName;
 
 	/** All names that contains a path to an UObject. */
 	TSet<FName> UObjectNames;
@@ -207,10 +272,13 @@ protected:
 	TMap<FName, TMap<uint64, FAllocationInfo> > SnapshotsWithAllocationMap;
 
 	/** Snapshots with callstack based allocation map. */
-	TMap<FName, TMap<FName, FSizeAndCount> > SnapshotsWithScopedAllocations;
+	TMap<FName, TMap<FName, FCombinedAllocationInfo> > SnapshotsWithScopedAllocations;
 
-	TMap<FName, TMap<FString, FSizeAndCount> > SnapshotsWithDecodedScopedAllocations;
+	TMap<FName, TMap<FString, FCombinedAllocationInfo> > SnapshotsWithDecodedScopedAllocations;
 
 	/** Filepath to the raw stats file. */
 	FString SourceFilepath;
+
+	/** Platform's name based on the loaded ue4statsraw file. */
+	FString PlatformName;
 };
