Stats - Refactored reading from a stats file (wip)

diff --git a/Engine/Source/Programs/UnrealFrontend/Private/Commands/StatsDumpMemoryCommand.cpp b/Engine/Source/Programs/UnrealFrontend/Private/Commands/StatsDumpMemoryCommand.cpp
--- a/Engine/Source/Programs/UnrealFrontend/Private/Commands/StatsDumpMemoryCommand.cpp
+++ b/Engine/Source/Programs/UnrealFrontend/Private/Commands/StatsDumpMemoryCommand.cpp
@@ -193,6 +193,7 @@ void FNodeAllocationInfo::PrepareCallstackData( const TArray<FName>& InDecodedCa
 /** Holds stats stack state, used to preserve continuity when the game frame has changed. */
 struct FStackState
 {
+	/** Default constructor. */
 	FStackState()
 		: bIsBrokenCallstack( false )
 	{}
@@ -217,186 +218,28 @@ static const double NumSecondsBetweenLogs = 5.0;
 
 void FStatsMemoryDumpCommand::InternalRun()
 {
+	FString SourceFilepath;
 	FParse::Value( FCommandLine::Get(), TEXT( "-INFILE=" ), SourceFilepath );
 
-	const int64 Size = IFileManager::Get().FileSize( *SourceFilepath );
-	if( Size < 4 )
+	FStatsReadFile* NewReader = FStatsReadFile::CreateReaderForRawStats( *SourceFilepath, nullptr );
+	if (NewReader)
 	{
-		UE_LOG( LogStats, Error, TEXT( "Could not open: %s" ), *SourceFilepath );
-		return;
-	}
-	TAutoPtr<FArchive> FileReader( IFileManager::Get().CreateFileReader( *SourceFilepath ) );
-	if( !FileReader )
-	{
-		UE_LOG( LogStats, Error, TEXT( "Could not open: %s" ), *SourceFilepath );
-		return;
-	}
-
-	FStatsReadStream Stream;
-
-	if( !Stream.ReadHeader( *FileReader ) )
-	{
-		UE_LOG( LogStats, Error, TEXT( "Could not open, bad magic: %s" ), *SourceFilepath );
-		return;
-	}
-
-	UE_LOG( LogStats, Warning, TEXT( "Reading a raw stats file for memory profiling: %s" ), *SourceFilepath );
-
-	const bool bIsFinalized = Stream.Header.IsFinalized();
-	check( bIsFinalized );
-	check( Stream.Header.Version >= EStatMagicWithHeader::VERSION_6 );
-
-	FStatsThreadState StatsThreadStats;
-	StatsThreadStats.MarkAsLoaded();
-
-	TArray<FStatMessage> Messages;
-	if( Stream.Header.bRawStatsFile )
-	{
-		FScopeLogTime SLT( TEXT( "FStatsMemoryDumpCommand::InternalRun" ), nullptr, FScopeLogTime::ScopeLog_Seconds );
-
-		// Read metadata.
-		TArray<FStatMessage> MetadataMessages;
-		Stream.ReadFNamesAndMetadataMessages( *FileReader, MetadataMessages );
-		StatsThreadStats.ProcessMetaDataOnly( MetadataMessages );
-
-		ThreadIdToName = StatsThreadStats.Threads;
-
-		// Find all UObject metadata messages.
-		for( const auto& Meta : MetadataMessages )
-		{
-			const FName EncName = Meta.NameAndInfo.GetEncodedName();
-			const FName RawName = Meta.NameAndInfo.GetRawName();
-			const FString Desc = FStatNameAndInfo::GetShortNameFrom( RawName ).GetPlainNameString();
-			const bool bContainsUObject = Desc.Contains( TEXT( "//" ) );
-			if( bContainsUObject )
-			{
-				UObjectNames.Add( RawName );
-			}
-		}
-
-		const int64 CurrentFilePos = FileReader->Tell();
-
-		// Read frames offsets.
-		Stream.ReadFramesOffsets( *FileReader );
-
-		// Buffer used to store the compressed and decompressed data.
-		TArray<uint8> SrcArray;
-		TArray<uint8> DestArray;
-		const bool bHasCompressedData = Stream.Header.HasCompressedData();
-		check( bHasCompressedData );
-
-		PlatformName = Stream.Header.PlatformName;
-
-		TMap<int64, FStatPacketArray> CombinedHistory;
-		int64 TotalDataSize = 0;
-		int64 TotalStatMessagesNum = 0;
-		int64 MaximumPacketSize = 0;
-		int64 TotalPacketsNum = 0;
-		// Read all packets sequentially, forced by the memory profiler which is now a part of the raw stats.
-		// !!CAUTION!! Frame number in the raw stats is pointless, because it is time/cycles based, not frame based.
-		// Background threads usually execute time consuming operations, so the frame number won't be valid.
-		// Needs to be combined by the thread and the time, not by the frame number.
-		{
-			// Display log information once per 5 seconds to avoid spamming.
-			double PreviousSeconds = FPlatformTime::Seconds();
-			const int64 FrameOffset0 = Stream.FramesInfo[0].FrameFileOffset;
-			FileReader->Seek( FrameOffset0 );
-
-			const int64 FileSize = FileReader->TotalSize();
-
-			while( FileReader->Tell() < FileSize )
-			{
-				// Read the compressed data.
-				FCompressedStatsData UncompressedData( SrcArray, DestArray );
-				*FileReader << UncompressedData;
-				if( UncompressedData.HasReachedEndOfCompressedData() )
-				{
-					break;
-				}
-
-				FMemoryReader MemoryReader( DestArray, true );
-
-				FStatPacket* StatPacket = new FStatPacket();
-				Stream.ReadStatPacket( MemoryReader, *StatPacket );
-
-				const int64 StatPacketFrameNum = StatPacket->Frame;
-				FStatPacketArray& Frame = CombinedHistory.FindOrAdd( StatPacketFrameNum );
-
-				// Check if we need to combine packets from the same thread.
-				FStatPacket** CombinedPacket = Frame.Packets.FindByPredicate( [&]( FStatPacket* Item ) -> bool
-				{
-					return Item->ThreadId == StatPacket->ThreadId;
-				} );
-
-				const int64 PacketSize = StatPacket->StatMessages.GetAllocatedSize();
-				TotalStatMessagesNum += StatPacket->StatMessages.Num();
-
-				if( CombinedPacket )
-				{
-					TotalDataSize -= (*CombinedPacket)->StatMessages.GetAllocatedSize();
-					(*CombinedPacket)->StatMessages += StatPacket->StatMessages;
-					TotalDataSize += (*CombinedPacket)->StatMessages.GetAllocatedSize();
-
-					delete StatPacket;
-				}
-				else
-				{
-					Frame.Packets.Add( StatPacket );
-					TotalDataSize += PacketSize;
-				}
-
-				const double CurrentSeconds = FPlatformTime::Seconds();
-				if( CurrentSeconds > PreviousSeconds + NumSecondsBetweenLogs )
-				{
-					const int32 PctPos = int32( 100.0*FileReader->Tell() / FileSize );
-					UE_LOG( LogStats, Log, TEXT( "%3i%% %10llu (%.1f MB) read messages, last read frame %4i" ), PctPos, TotalStatMessagesNum, TotalDataSize / 1024.0f / 1024.0f, StatPacketFrameNum );
-					PreviousSeconds = CurrentSeconds;
-				}
-			
-				MaximumPacketSize = FMath::Max( MaximumPacketSize, PacketSize );			
-				TotalPacketsNum++;
-			}
-		}
-
-		// Dump frame stats
-		for( const auto& It : CombinedHistory )
-		{
-			const int64 FrameNum = It.Key;
-			int64 FramePacketsSize = 0;
-			int64 FrameStatMessages = 0;
-			int64 FramePackets = It.Value.Packets.Num(); // Threads
-			for( const auto& It2 : It.Value.Packets )
-			{
-				FramePacketsSize += It2->StatMessages.GetAllocatedSize();
-				FrameStatMessages += It2->StatMessages.Num();
-			}
-
-			UE_LOG( LogStats, Warning, TEXT( "Frame: %10llu/%3lli Size: %.1f MB / %10lli" ), 
-					FrameNum, 
-					FramePackets, 
-					FramePacketsSize / 1024.0f / 1024.0f,
-					FrameStatMessages );
-		}
-
-		UE_LOG( LogStats, Warning, TEXT( "TotalPacketSize: %.1f MB, Max: %1f MB" ),
-				TotalDataSize / 1024.0f / 1024.0f,
-				MaximumPacketSize / 1024.0f / 1024.0f );
-
-		TArray<int64> Frames;
-		CombinedHistory.GenerateKeyArray( Frames );
-		Frames.Sort();
-		const int64 MiddleFrame = Frames[Frames.Num() / 2];
-
-		ProcessMemoryOperations( CombinedHistory );
+		NewReader->ReadSynchronously();
+		PlatformName = NewReader->Header.PlatformName;
+		ProcessMemoryOperations( *NewReader );
+		delete NewReader;
 	}
 }
 
 
-void FStatsMemoryDumpCommand::ProcessMemoryOperations( const TMap<int64, FStatPacketArray>& CombinedHistory )
+void FStatsMemoryDumpCommand::ProcessMemoryOperations( const FStatsReadFile& File )
 {
+	const TMap<int64, FStatPacketArray>& CombinedHistory = File.CombinedHistory;
+
 	// This is only example code, no fully implemented, may sometimes crash.
 	// This code is not optimized. 
 	double PreviousSeconds = FPlatformTime::Seconds();
+	int32 CurrentStatMessageIndex = 0;
 	uint64 NumMemoryOperations = 0;
 
 	// Generate frames
@@ -423,15 +266,6 @@ void FStatsMemoryDumpCommand::ProcessMemoryOperations( const TMap<int64, FStatPa
 
 	for( int32 FrameIndex = 0; FrameIndex < Frames.Num(); ++FrameIndex )
 	{
-		{
-			const double CurrentSeconds = FPlatformTime::Seconds();
-			if( CurrentSeconds > PreviousSeconds + NumSecondsBetweenLogs )
-			{
-				UE_LOG( LogStats, Warning, TEXT( "Processing frame %i/%i" ), FrameIndex+1, Frames.Num() );
-				PreviousSeconds = CurrentSeconds;
-			}
-		}
-
 		const int64 TargetFrame = Frames[FrameIndex];
 		const int64 Diff = TargetFrame - FirstFrame;
 		const FStatPacketArray& Frame = CombinedHistory.FindChecked( TargetFrame );
@@ -439,18 +273,8 @@ void FStatsMemoryDumpCommand::ProcessMemoryOperations( const TMap<int64, FStatPa
 		bool bAtLeastOnePacket = false;
 		for( int32 PacketIndex = 0; PacketIndex < Frame.Packets.Num(); PacketIndex++ )
 		{
-			{
-				const double CurrentSeconds = FPlatformTime::Seconds();
-				if( CurrentSeconds > PreviousSeconds + NumSecondsBetweenLogs )
-				{
-					UE_LOG( LogStats, Log, TEXT( "Processing packet %i/%i" ), PacketIndex, Frame.Packets.Num() );
-					PreviousSeconds = CurrentSeconds;
-					bAtLeastOnePacket = true;
-				}
-			}
-
 			const FStatPacket& StatPacket = *Frame.Packets[PacketIndex];
-			const FName& ThreadFName = ThreadIdToName.FindChecked( StatPacket.ThreadId );
+			const FName& ThreadFName = File.State.Threads.FindChecked( StatPacket.ThreadId );
 
 			FStackState* StackState = StackStates.Find( ThreadFName );
 			if( !StackState )
@@ -462,26 +286,24 @@ void FStatsMemoryDumpCommand::ProcessMemoryOperations( const TMap<int64, FStatPa
 
 			const FStatMessagesArray& Data = StatPacket.StatMessages;
 
-			int32 LastPct = 0;
-			const int32 NumDataElements = Data.Num();
-			const int32 OnerPercent = FMath::Max( NumDataElements / 100, 1024 );
-			bool bAtLeastOneMessage = false;
-			for( int32 Index = 0; Index < NumDataElements; Index++ )
+			const int32 NumStatMessages = Data.Num();
+			const int32 OnerPercent = FMath::Max( NumStatMessages / 100, 1024 );
+			for( int32 Index = 0; Index < NumStatMessages; Index++ )
 			{
+				CurrentStatMessageIndex++;
+
 				if( Index % OnerPercent )
 				{
 					const double CurrentSeconds = FPlatformTime::Seconds();
 					if( CurrentSeconds > PreviousSeconds + NumSecondsBetweenLogs )
 					{
-						const int32 CurrentPct = int32( 100.0*(Index + 1) / NumDataElements );
-						UE_LOG( LogStats, Log, TEXT( "Processing %3i%% (%i/%i) stat messages" ), CurrentPct, Index, NumDataElements );
+						const int32 CurrentPct = int32( 100.0*CurrentStatMessageIndex / NumStatMessages );
+						UE_LOG( LogStats, Log, TEXT( "Processing %3i%% (%i/%i) stat messages [Frame: %3i, Packet: %2i]" ), CurrentPct, Index, NumStatMessages, FrameIndex, PacketIndex );
 						PreviousSeconds = CurrentSeconds;
-						bAtLeastOneMessage = true;
 					}
 				}
 
 				const FStatMessage& Item = Data[Index];
-
 				const EStatOperation::Type Op = Item.NameAndInfo.GetField<EStatOperation>();
 				const FName RawName = Item.NameAndInfo.GetRawName();
 
@@ -638,14 +460,6 @@ void FStatsMemoryDumpCommand::ProcessMemoryOperations( const TMap<int64, FStatPa
 					}
 				}
 			}
-			if( bAtLeastOneMessage )
-			{
-				PreviousSeconds -= NumSecondsBetweenLogs;
-			}
-		}
-		if( bAtLeastOnePacket )
-		{
-			PreviousSeconds -= NumSecondsBetweenLogs;
 		}
 	}
 
@@ -964,7 +778,7 @@ void FStatsMemoryDumpCommand::GenerateScopedTreeAllocations( const TMap<FName, F
 }
 
 
-void FStatsMemoryDumpCommand::GenerateMemoryUsageReport( const TMap<uint64, FAllocationInfo>& AllocationMap )
+void FStatsMemoryDumpCommand::GenerateMemoryUsageReport( const TMap<uint64, FAllocationInfo>& AllocationMap, const TSet<FName>& UObjectRawNames )
 {
 	if( AllocationMap.Num() == 0 )
 	{
@@ -1049,7 +863,7 @@ void FStatsMemoryDumpCommand::ProcessAndDumpScopedAllocations( const TMap<uint64
 	MemoryReport.CycleRow();
 }
 
-void FStatsMemoryDumpCommand::ProcessAndDumpUObjectAllocations( const TMap<uint64, FAllocationInfo>& AllocationMap )
+void FStatsMemoryDumpCommand::ProcessAndDumpUObjectAllocations( const TMap<uint64, FAllocationInfo>& AllocationMap, const TSet<FName>& UObjectRawNames )
 {
 	// This code is not optimized. 
 	FScopeLogTime SLT( TEXT( "ProcessingUObjectAllocations" ), nullptr, FScopeLogTime::ScopeLog_Seconds );
@@ -1085,7 +899,7 @@ void FStatsMemoryDumpCommand::ProcessAndDumpUObjectAllocations( const TMap<uint6
 			for( int32 Index = DecodedCallstack.Num() - 1; Index >= 0; --Index )
 			{
 				const FName LongName = DecodedCallstack[Index];
-				const bool bValid = UObjectNames.Contains( LongName );
+				const bool bValid = UObjectRawNames.Contains( LongName );
 				if( bValid )
 				{
 					const FString ObjectName = FStatNameAndInfo::GetShortNameFrom( LongName ).GetPlainNameString();
diff --git a/Engine/Source/Programs/UnrealFrontend/Private/Commands/StatsDumpMemoryCommand.h b/Engine/Source/Programs/UnrealFrontend/Private/Commands/StatsDumpMemoryCommand.h
--- a/Engine/Source/Programs/UnrealFrontend/Private/Commands/StatsDumpMemoryCommand.h
+++ b/Engine/Source/Programs/UnrealFrontend/Private/Commands/StatsDumpMemoryCommand.h
@@ -223,16 +223,16 @@ protected:
 	void InternalRun();
 
 	/** Basic memory profiling, only for debugging purpose. */
-	void ProcessMemoryOperations( const TMap<int64, FStatPacketArray>& CombinedHistory );
+	void ProcessMemoryOperations( const FStatsReadFile& File );
 
 	/** Generates a basic memory usage report and prints it to the log. */
-	void GenerateMemoryUsageReport( const TMap<uint64, FAllocationInfo>& AllocationMap );
+	void GenerateMemoryUsageReport( const TMap<uint64, FAllocationInfo>& AllocationMap, const TSet<FName>& UObjectRawNames );
 
 	/** Generates scoped allocation statistics. */
 	void ProcessAndDumpScopedAllocations( const TMap<uint64, FAllocationInfo>& AllocationMap );
 
 	/** Generates UObject allocation statistics. */
-	void ProcessAndDumpUObjectAllocations( const TMap<uint64, FAllocationInfo>& AllocationMap );
+	void ProcessAndDumpUObjectAllocations( const TMap<uint64, FAllocationInfo>& AllocationMap, const TSet<FName>& UObjectRawNames );
 
 	void DumpScopedAllocations( const TCHAR* Name, const TMap<FString, FCombinedAllocationInfo>& CombinedAllocations );
 
@@ -256,12 +256,6 @@ protected:
 		return PlatformName;
 	}
 
-	/** Map from a thread id to the thread name. */
-	TMap<uint32, FName> ThreadIdToName;
-
-	/** All names that contains a path to an UObject. */
-	TSet<FName> UObjectNames;
-
 	/** The sequence tag mapping to the named markers. */
 	TArray<TPair<uint32, FName>> Snapshots;
 
@@ -276,9 +270,6 @@ protected:
 
 	TMap<FName, TMap<FString, FCombinedAllocationInfo> > SnapshotsWithDecodedScopedAllocations;
 
-	/** Filepath to the raw stats file. */
-	FString SourceFilepath;
-
 	/** Platform's name based on the loaded ue4statsraw file. */
 	FString PlatformName;
 };
diff --git a/Engine/Source/Runtime/Core/Private/Stats/StatsFile.cpp b/Engine/Source/Runtime/Core/Private/Stats/StatsFile.cpp
--- a/Engine/Source/Runtime/Core/Private/Stats/StatsFile.cpp
+++ b/Engine/Source/Runtime/Core/Private/Stats/StatsFile.cpp
@@ -426,12 +426,12 @@ void IStatsWriteFile::SendTask()
 void FStatsWriteFile::SetDataDelegate( bool bSet )
 {
 	FStatsThreadState const& Stats = FStatsThreadState::GetLocalState();
-	if( bSet )
+	if (bSet)
 	{
 		DataDelegateHandle = Stats.NewFrameDelegate.AddRaw( this, &FStatsWriteFile::WriteFrame );
-}
+	}
 	else
-{
+	{
 		Stats.NewFrameDelegate.Remove( DataDelegateHandle );
 	}
 }
@@ -523,6 +523,257 @@ void FRawStatsWriteFile::WriteStatPacket( FArchive& Ar, FStatPacket& StatPacket
 	}
 }
 
+const double FStatsReadFile::NumSecondsBetweenUpdates = 5.0;
+
+FStatsReadFile* FStatsReadFile::CreateReaderForRegularStats( const TCHAR* Filename, FOnNewStatPacket InNewStatPacket )
+{
+	FStatsReadFile* StatsReadFile = new FStatsReadFile( Filename, false );
+	const bool bValid = StatsReadFile->PrepareLoading();
+	if (bValid)
+	{
+		StatsReadFile->NewStatPacket = InNewStatPacket;
+	}
+	return bValid ? StatsReadFile : nullptr;
+}
+
+
+FStatsReadFile* FStatsReadFile::CreateReaderForRawStats( const TCHAR* Filename, FOnNewRawStatPacket InNewRawStatPacket )
+{
+	FStatsReadFile* StatsReadFile = new FStatsReadFile( Filename, true );
+	const bool bValid = StatsReadFile->PrepareLoading();
+	if (bValid)
+	{
+		StatsReadFile->NewRawStatPacket = InNewRawStatPacket;
+	}
+	return bValid ? StatsReadFile : nullptr;
+}
+
+void FStatsReadFile::ReadSynchronously()
+{
+	const double StartTime = FPlatformTime::Seconds();
+
+	if (bRawStatsFile)
+	{
+		const int64 FileSize = Reader->TotalSize();
+
+		// Buffer used to store the compressed and decompressed data.
+		TArray<uint8> SrcArray;
+		TArray<uint8> DestArray;
+
+		// Read all packets sequentially, forced by the memory profiler which is now a part of the raw stats.
+		// !!CAUTION!! Frame number in the raw stats is pointless, because it is time/cycles based, not frame based.
+		// Background threads usually execute time consuming operations, so the frame number won't be valid.
+		// Needs to be combined by the thread and the time, not by the frame number.
+
+		// Update stage progress once per NumSecondsBetweenUpdates(5) seconds to avoid spamming.
+		double PreviousSeconds = FPlatformTime::Seconds();
+		ReadingStage.Set( int32( ERawStatsReadingStage::RS_ReadingAndCombiningPackets ) );
+
+		while (Reader->Tell() < FileSize)
+		{
+			// Read the compressed data.
+			FCompressedStatsData UncompressedData( SrcArray, DestArray );
+			*Reader << UncompressedData;
+			if (UncompressedData.HasReachedEndOfCompressedData())
+			{
+				StageProgress.Set( 100 );
+				break;
+			}
+
+			FMemoryReader MemoryReader( DestArray, true );
+
+			FStatPacket* StatPacket = new FStatPacket();
+			Stream.ReadStatPacket( MemoryReader, *StatPacket );
+
+			const int64 StatPacketFrameNum = StatPacket->Frame;
+			FStatPacketArray& Frame = CombinedHistory.FindOrAdd( StatPacketFrameNum );
+
+			// Check if we need to combine packets from the same thread.
+			FStatPacket** CombinedPacket = Frame.Packets.FindByPredicate( [&]( FStatPacket* Item ) -> bool
+			{
+				return Item->ThreadId == StatPacket->ThreadId;
+			} );
+
+			const int64 PacketSize = StatPacket->StatMessages.GetAllocatedSize();
+			FileInfo.TotalStatMessagesNum += StatPacket->StatMessages.Num();
+
+			if (CombinedPacket)
+			{
+				FileInfo.TotalPacketsSize -= (*CombinedPacket)->StatMessages.GetAllocatedSize();
+				(*CombinedPacket)->StatMessages += StatPacket->StatMessages;
+				FileInfo.TotalPacketsSize += (*CombinedPacket)->StatMessages.GetAllocatedSize();
+
+				delete StatPacket;
+			}
+			else
+			{
+				Frame.Packets.Add( StatPacket );
+				FileInfo.TotalPacketsSize += PacketSize;
+				FileInfo.MaximumPacketSize = FMath::Max( FileInfo.MaximumPacketSize, PacketSize );
+			}
+
+			const double CurrentSeconds = FPlatformTime::Seconds(); 			
+			if (CurrentSeconds > LastUpdateTime + NumSecondsBetweenUpdates)
+			{
+				const int32 PercentagePos = int32( 100.0*Reader->Tell() / FileSize );
+				StageProgress.Set( PercentagePos );
+				UE_LOG( LogStats, Verbose, TEXT( "%3i%% %10llu (%6.1f MB) read messages, last read frame %4i" ), PercentagePos, FileInfo.TotalStatMessagesNum, FileInfo.TotalPacketsSize / 1024.0f / 1024.0f, StatPacketFrameNum );
+				LastUpdateTime = CurrentSeconds;
+			}
+			
+			FileInfo.TotalPacketsNum++;
+		}
+
+		const double TotalTime = FPlatformTime::Seconds() - StartTime;
+		UE_LOG( LogStats, Log, TEXT( "Read took %.2f sec(s)" ), TotalTime );
+
+		DumpCombinedHistoryStats();	
+
+		//system( "pause" );
+
+		ReadingStage.Set( int32( ERawStatsReadingStage::RS_DecodingStatMessages ) );
+	}
+}
+
+void FStatsReadFile::ReadAsynchronously()
+{
+
+}
+
+FStatsReadFile::FStatsReadFile( const TCHAR* InFilename, bool bInRawStatsFile )
+	: Header( Stream.Header )
+	, Reader( nullptr )
+	, LastUpdateTime( 0.0 )
+	, Filename( InFilename )
+	, bRawStatsFile( bInRawStatsFile )
+{
+	
+}
+
+FStatsReadFile::~FStatsReadFile()
+{
+
+}
+
+bool FStatsReadFile::PrepareLoading()
+{
+	const double StartTime = FPlatformTime::Seconds();
+
+	const int64 Size = IFileManager::Get().FileSize( *Filename );
+	if (Size < 4)
+	{
+		UE_LOG( LogStats, Error, TEXT( "Could not open: %s" ), *Filename );
+		return false;
+	}
+
+	Reader = IFileManager::Get().CreateFileReader( *Filename );
+	if (!Reader)
+	{
+		UE_LOG( LogStats, Error, TEXT( "Could not open: %s" ), *Filename );
+		return false;
+	}
+
+	if (!Stream.ReadHeader( *Reader ))
+	{
+		UE_LOG( LogStats, Error, TEXT( "Could not read, header is invalid: %s" ), *Filename );
+		return false;
+	}
+
+	//UE_LOG( LogStats, Warning, TEXT( "Reading a raw stats file for memory profiling: %s" ), *Filename );
+
+	const bool bIsFinalized = Stream.Header.IsFinalized();
+	if (!bIsFinalized)
+	{
+		UE_LOG( LogStats, Error, TEXT( "Could not read, file is not finalized: %s" ), *Filename );
+		return false;
+	}
+
+	if( Stream.Header.Version < EStatMagicWithHeader::VERSION_6 )
+	{
+		UE_LOG( LogStats, Error, TEXT( "Could not read, invalid version: %s, expected %u, was %u" ), *Filename, (uint32)EStatMagicWithHeader::VERSION_6, Stream.Header.Version );
+		return false;
+	}
+
+	const bool bHasCompressedData = Stream.Header.HasCompressedData();
+	if( !bHasCompressedData )
+	{
+		UE_LOG( LogStats, Error, TEXT( "Could not read, required compressed data: %s" ), *Filename );
+		return false;
+	}
+
+	State.MarkAsLoaded();
+
+	
+	// Read metadata.
+	TArray<FStatMessage> MetadataMessages;
+	Stream.ReadFNamesAndMetadataMessages( *Reader, MetadataMessages );
+	State.ProcessMetaDataOnly( MetadataMessages );
+
+	// Find all UObject metadata messages.
+	for (const auto& Meta : MetadataMessages)
+	{
+		const FName EncName = Meta.NameAndInfo.GetEncodedName();
+		const FName RawName = Meta.NameAndInfo.GetRawName();
+		const FString Desc = FStatNameAndInfo::GetShortNameFrom( RawName ).GetPlainNameString();
+		const bool bContainsUObject = Desc.Contains( TEXT( "//" ) );
+		if (bContainsUObject)
+		{
+			UObjectRawNames.Add( RawName );
+		}
+	}
+
+	// Read frames offsets.
+	Stream.ReadFramesOffsets( *Reader );
+
+	// Move file pointer to the first frame or first stat packet.
+	const int64 FrameOffset0 = Stream.FramesInfo[0].FrameFileOffset;
+	Reader->Seek( FrameOffset0 );
+	const int64 FileSize = Reader->TotalSize();
+
+	const double TotalTime = FPlatformTime::Seconds() - StartTime;
+	UE_LOG( LogStats, Log, TEXT( "Prepare loading took %.2f sec(s)" ), TotalTime );
+
+	return true;
+}
+
+void FStatsReadFile::DumpCombinedHistoryStats()
+{
+	// Dump frame stats
+	int64 LastFrameNum = 0;
+	for (const auto& It : CombinedHistory)
+	{
+		const int64 FrameNum = It.Key;
+		int64 FramePacketsSize = 0;
+		int64 FrameStatMessages = 0;
+		int64 FramePackets = It.Value.Packets.Num(); // Threads
+		for (const auto& It2 : It.Value.Packets)
+		{
+			FramePacketsSize += It2->StatMessages.GetAllocatedSize();
+			FrameStatMessages += It2->StatMessages.Num();
+		}
+
+		UE_LOG( LogStats, Verbose, TEXT( "Frame: %4llu/%2lli Size: %5.1f MB / %10lli" ),
+				FrameNum,
+				FramePackets,
+				FramePacketsSize / 1024.0f / 1024.0f,
+				FrameStatMessages );
+
+		LastFrameNum = FMath::Max( LastFrameNum, FrameNum );
+	}
+
+	UE_LOG( LogStats, Warning, TEXT( "TotalPacketSize: %6.1f MB, Max: %2f MB, TotalPacketsNum: %lli, Frames: %i" ),
+			FileInfo.TotalPacketsSize / 1024.0f / 1024.0f,
+			FileInfo.MaximumPacketSize / 1024.0f / 1024.0f,
+			FileInfo.TotalPacketsNum,
+			CombinedHistory.Num() );
+
+	TArray<int64> Frames;
+	CombinedHistory.GenerateKeyArray( Frames );
+	Frames.Sort();
+	// Verify that frames are sequential.
+	check( Frames[Frames.Num() - 1] == LastFrameNum );
+}
+
 /*-----------------------------------------------------------------------------
 	Commands functionality
 -----------------------------------------------------------------------------*/;
@@ -535,13 +786,13 @@ FCommandStatsFile& FCommandStatsFile::Get()
 
 
 void FCommandStatsFile::Start( const FString& Filename )
-	{
+{
 	Stop();
 	CurrentStatsFile = new FStatsWriteFile();
 	CurrentStatsFile->Start( Filename );
 
 	StatFileActiveCounter.Increment();
-	}
+}
 
 void FCommandStatsFile::StartRaw( const FString& Filename )
 {
@@ -554,7 +805,7 @@ void FCommandStatsFile::StartRaw( const FString& Filename )
 
 void FCommandStatsFile::Stop()
 {
-	if( CurrentStatsFile )
+	if (CurrentStatsFile)
 	{
 		CurrentStatsFile->Stop();
 		delete CurrentStatsFile;
diff --git a/Engine/Source/Runtime/Core/Public/Stats/StatsFile.h b/Engine/Source/Runtime/Core/Public/Stats/StatsFile.h
--- a/Engine/Source/Runtime/Core/Public/Stats/StatsFile.h
+++ b/Engine/Source/Runtime/Core/Public/Stats/StatsFile.h
@@ -10,53 +10,47 @@ class FAsyncStatsWrite;
 /**
 * Magic numbers for stats streams, this is for the first version.
 */
-namespace EStatMagicNoHeader
+enum EStatMagicNoHeader : uint32
 {
-	enum Type
-	{
-		MAGIC = 0x7E1B83C1,
-		MAGIC_SWAPPED = 0xC1831B7E,
-		NO_VERSION = 0,
-	};
-}
+	MAGIC_NO_HEADER = 0x7E1B83C1,
+	MAGIC_NO_HEADER_SWAPPED = 0xC1831B7E,
+	NO_VERSION = 0,
+};
 
 /**
 * Magic numbers for stats streams, this is for the second and later versions.
 * Allows more advanced options.
 */
-namespace EStatMagicWithHeader
+enum EStatMagicWithHeader : uint32
 {
-	enum Type
-	{
-		MAGIC = 0x10293847,
-		MAGIC_SWAPPED = 0x47382910,
-		VERSION_2 = 2,
-		VERSION_3 = 3,
-
-		/**
-		 *	Added support for compressing the stats data, now each frame is compressed.
-		 *	!!CAUTION!! Deprecated.
-		 */
-		VERSION_4 = 4,
-		HAS_COMPRESSED_DATA_VER = VERSION_4,
+	MAGIC = 0x10293847,
+	MAGIC_SWAPPED = 0x47382910,
+	VERSION_2 = 2,
+	VERSION_3 = 3,
 
-		/**
-		 *	New low-level raw stats with memory profiler functionality.
-		 *  !!CAUTION!! Not backward compatible with version 4.
-		 */
-		VERSION_5 = 5,
+	/**
+	*	Added support for compressing the stats data, now each frame is compressed.
+	*	!!CAUTION!! Deprecated.
+	*/
+	VERSION_4 = 4,
+	HAS_COMPRESSED_DATA_VER = VERSION_4,
 
-		/**
-		 *	Added realloc message to avoid sending alloc and free
-		 *	Should also reduce the amount of all messages
-		 *  !!CAUTION!! Not backward compatible with version 5.
-		 */
-		VERSION_6 = 6,
+	/**
+	*	New low-level raw stats with memory profiler functionality.
+	*  !!CAUTION!! Not backward compatible with version 4.
+	*/
+	VERSION_5 = 5,
 
-		/** Latest version. */
-		VERSION_LATEST = VERSION_6,
-	};
-}
+	/**
+	*	Added realloc message to avoid sending alloc and free
+	*	Should also reduce the amount of all messages
+	*  !!CAUTION!! Not backward compatible with version 5.
+	*/
+	VERSION_6 = 6,
+
+	/** Latest version. */
+	VERSION_LATEST = VERSION_6,
+};
 
 struct EStatsFileConstants
 {
@@ -559,11 +553,11 @@ public:
 
 		uint32 Magic = 0;
 		Ar << Magic;
-		if( Magic == EStatMagicNoHeader::MAGIC )
+		if( Magic == EStatMagicNoHeader::MAGIC_NO_HEADER )
 		{
 
 		}
-		else if( Magic == EStatMagicNoHeader::MAGIC_SWAPPED )
+		else if( Magic == EStatMagicNoHeader::MAGIC_NO_HEADER_SWAPPED )
 		{
 			Ar.SetByteSwapping( true );
 		}
@@ -590,9 +584,7 @@ public:
 		return true;
 	}
 
-	// @TODO yrx 2014-12-02 Add a better way to read the file.
-
-	/** Reads a stat packed from the specified archive. */
+	/** Reads a stat packed from the specified archive. Only for raw stats files. */
 	void ReadStatPacket( FArchive& Ar, FStatPacket& StatPacked )
 	{
 		Ar << StatPacked.Frame;
@@ -751,6 +743,147 @@ public:
 	}
 };
 
+/** Raw stats information. */
+struct FRawStatsFileInfo
+{
+	/** Default constructor. */
+	FRawStatsFileInfo()
+		: TotalPacketsSize( 0 )
+		, TotalStatMessagesNum( 0 )
+		, MaximumPacketSize( 0 )
+		, TotalPacketsNum( 0 )
+	{}
+
+	/** Size of all packets. */
+	int64 TotalPacketsSize;
+
+	/** Number of all stat messages. */
+	int64 TotalStatMessagesNum;
+
+	/** Maximum packet size. */
+	int64 MaximumPacketSize;
+
+	/** Number of all packets. */
+	int64 TotalPacketsNum;
+};
+
+/** Enumerates raw stats reading stages. */
+enum ERawStatsReadingStage
+{
+	RS_ReadingAndCombiningPackets = 0,
+	RS_DecodingStatMessages,
+	/** Memory profiler only. */
+	RS_SortingSequence,
+
+};
+
+/** Struct used to read from ue4stats/ue4statsraw files, initializes all metadata and starts a process of reading the file asynchronously. */
+// #YRX_Stats: 2015-07-08 Maybe add IStatsReadFile -> FStatsReadFile, FRawStatsReadFile
+// #YRX_Stats: 2015-07-09 Remove CORE_API
+struct CORE_API FStatsReadFile
+{
+protected:
+	/** Number of seconds between updating the current stage. */
+	static const double NumSecondsBetweenUpdates;
+
+	/** Delegate that FStatsReadFile calls on the instance whenever we have a new frame. */
+	DECLARE_DELEGATE_OneParam( FOnNewStatPacket, TArray<FStatMessage>* );
+
+	/** Delegate that FStatsReadFile calls on the instance whenever we have a new raw stats packet. */
+	DECLARE_DELEGATE_OneParam( FOnNewRawStatPacket, FStatPacket* );
+
+public:
+	/** Creates a new reader for regular stats file. Will be nullptr for invalid files. */
+	static FStatsReadFile* CreateReaderForRegularStats( const TCHAR* Filename, FOnNewStatPacket InNewStatPacket );
+
+	/** Creates a new reader for raw stats file. Will be nullptr for invalid files. */
+	static FStatsReadFile* CreateReaderForRawStats( const TCHAR* Filename, FOnNewRawStatPacket InNewRawStatPacket );
+
+	/** Read the file on the current thread. This is a blocking operation. */
+	void ReadSynchronously();
+
+	/** Read the file using the async tasks on the pool thread. The read data is sent to the game thread using the task graph. This is a non-blocking operation. */
+	void ReadAsynchronously();
+
+	/**	 
+	 * @return true, if any asynchronous operation is being processed
+	 */
+	bool IsBusy()
+	{
+		// #YRX_STATS: 2015-07-09 
+		return false;
+	}
+	
+	/**
+	 * @return current stage of reading.
+	 */
+	const ERawStatsReadingStage GetReadingStage() const
+	{
+		return (ERawStatsReadingStage)ReadingStage.GetValue();
+	}
+
+protected:
+	/** Initialization constructor. */
+	FStatsReadFile( const TCHAR* InFilename, bool bInRawStatsFile );
+
+public:
+	/** Destructor. */
+	~FStatsReadFile();
+
+protected:
+	/**
+	 * Prepares file to be loaded, makes sanity checks, reads and initializes metadata
+	 * @return true, if the process was completed successfully
+	 */
+	bool PrepareLoading();
+
+	/** Dumps combined history stats. Only for raw stats. */
+	void DumpCombinedHistoryStats();
+
+//protected:
+public:
+	/** Current state of the stats. Mostly for metadata. */
+	FStatsThreadState State;
+
+	/** Reading stream. */
+	FStatsReadStream Stream;
+
+	/** Reference to the stream's header. */
+	FStatsStreamHeader& Header;
+
+	/** File reader. */
+	FArchive* Reader;
+
+	/** Delegate we fire every time we read one frame of data. The data must be freed after processing. */
+	mutable FOnNewStatPacket NewStatPacket;
+
+	/** Delegate we fire every time we read one raw stats packet. The data must be freed after processing. */
+	mutable FOnNewRawStatPacket NewRawStatPacket;
+
+	/** Basic information about the stats file. */
+	FRawStatsFileInfo FileInfo;
+
+	/** Combined history for raw packets, indexed by a frame number. */
+	TMap<int64, FStatPacketArray> CombinedHistory;
+
+	/** All raw names that contains a path to an UObject. */
+	TSet<FName> UObjectRawNames;
+
+	/** Current stage of reading. Thread safe. */
+	FThreadSafeCounter ReadingStage;
+
+	/** Percentage progress of the current stage. */
+	FThreadSafeCounter StageProgress;
+
+	/** Time of the last stage update. */
+	double LastUpdateTime;
+
+	/** Filename. */
+	const FString Filename;
+
+	/** Whether this stats file uses raw data. */
+	const bool bRawStatsFile;
+};
 
 /*-----------------------------------------------------------------------------
 	Commands functionality
@@ -769,10 +902,10 @@ struct FCommandStatsFile
 	{}
 
 	/** Stat StartFile. */
-	void Start( const FString& File );
+	void Start( const FString& Filename );
 
 	/** Stat StartFileRaw. */
-	void StartRaw( const FString& File );
+	void StartRaw( const FString& Filename );
 
 	/** Stat StopFile. */
 	void Stop();
