Stats - Refactored reading from a stats file (asynchronous read/process raw stats) (WIP)

diff --git a/Engine/Source/Programs/UnrealFrontend/Private/Commands/StatsDumpMemoryCommand.cpp b/Engine/Source/Programs/UnrealFrontend/Private/Commands/StatsDumpMemoryCommand.cpp
--- a/Engine/Source/Programs/UnrealFrontend/Private/Commands/StatsDumpMemoryCommand.cpp
+++ b/Engine/Source/Programs/UnrealFrontend/Private/Commands/StatsDumpMemoryCommand.cpp
@@ -225,10 +225,53 @@ void FStatsMemoryDumpCommand::InternalRun()
 	if (NewReader)
 	{
 		PlatformName = NewReader->Header.PlatformName;
-		NewReader->ReadAndProcessSynchronously();
+		NewReader->ReadAndProcessAsynchronously();
+
+		while( NewReader->IsBusy() )
+		{
+			FPlatformProcess::Sleep( 1.0f );
+
+			UE_LOG( LogStats, Log, TEXT( "Async: Stage: %s/%3i%%" ), *NewReader->GetProcessingStageAsString(), NewReader->GetStageProgress() );
+		}
+
 		delete NewReader;
 
+		// Frame-240 Frame-120 Frame-060
+		TMap<FString, FCombinedAllocationInfo> FrameBegin_Exit;
+		CompareSnapshots_FString( TEXT( "BeginSnapshot" ), TEXT( "EngineLoop.Exit" ), FrameBegin_Exit );
+		DumpScopedAllocations( TEXT( "Begin_Exit" ), FrameBegin_Exit );
+
+#if	UE_BUILD_DEBUG
+		TMap<FString, FCombinedAllocationInfo> Frame060_120;
+		CompareSnapshots_FString( TEXT( "Frame-060" ), TEXT( "Frame-120" ), Frame060_120 );
+		DumpScopedAllocations( TEXT( "Frame060_120" ), Frame060_120 );
+
+		TMap<FString, FCombinedAllocationInfo> Frame060_240;
+		CompareSnapshots_FString( TEXT( "Frame-060" ), TEXT( "Frame-240" ), Frame060_240 );
+		DumpScopedAllocations( TEXT( "Frame060_240" ), Frame060_240 );
+
+		// Generate scoped tree view.
+		{
+			TMap<FName, FCombinedAllocationInfo> FrameBegin_Exit_FName;
+			CompareSnapshots_FName( TEXT( "BeginSnapshot" ), TEXT( "EngineLoop.Exit" ), FrameBegin_Exit_FName );
+
+			FNodeAllocationInfo Root;
+			Root.EncodedCallstack = TEXT( "ThreadRoot" );
+			Root.HumanReadableCallstack = TEXT( "ThreadRoot" );
+			GenerateScopedTreeAllocations( FrameBegin_Exit_FName, Root );
+		}
+
+
+	{
+		TMap<FName, FCombinedAllocationInfo> Frame060_240_FName;
+		CompareSnapshots_FName( TEXT( "Frame-060" ), TEXT( "Frame-240" ), Frame060_240_FName );
 
+		FNodeAllocationInfo Root;
+		Root.EncodedCallstack = TEXT( "ThreadRoot" );
+		Root.HumanReadableCallstack = TEXT( "ThreadRoot" );
+		GenerateScopedTreeAllocations( Frame060_240_FName, Root );
+	}
+#endif // UE_BUILD_DEBUG
 	}
 }
 
@@ -309,7 +352,7 @@ void FStatsMemoryDumpCommand::ProcessMemoryOperations( FStatsReadFile* InFile )
 						// Abandon support, simply return for now.
 						if (InFile->bShouldStopProcessing == true)
 						{
-							InFile->SetProcessingStage( EStatsProcessingStage::RS_Stopped );
+							InFile->SetProcessingStage( EStatsProcessingStage::SPS_Stopped );
 							return;
 						}
 					}
@@ -507,7 +550,7 @@ void FStatsMemoryDumpCommand::ProcessMemoryOperations( FStatsReadFile* InFile )
 	}
 	*/
 
-	InFile->SetProcessingStage( EStatsProcessingStage::RS_SortSequence );
+	InFile->SetProcessingStage( EStatsProcessingStage::SPS_SortSequence );
 
 	// Sort all memory operation by the sequence tag, iterate through all operation and generate memory usage.
 	SequenceAllocationArray.Sort( FAllocationInfoSequenceTagLess() );
@@ -515,7 +558,7 @@ void FStatsMemoryDumpCommand::ProcessMemoryOperations( FStatsReadFile* InFile )
 	// Abandon support, simply return for now.
 	if (InFile->bShouldStopProcessing == true)
 	{
-		InFile->SetProcessingStage( EStatsProcessingStage::RS_Stopped );
+		InFile->SetProcessingStage( EStatsProcessingStage::SPS_Stopped );
 		return;
 	}
 
@@ -529,6 +572,8 @@ void FStatsMemoryDumpCommand::ProcessMemoryOperations( FStatsReadFile* InFile )
 	int32 NumFWAMemoryOperations = 0; // FreeWithoutAlloc
 	int32 NumZeroAllocs = 0; // Malloc(0)
 
+	InFile->SetProcessingStage( EStatsProcessingStage::SPS_ProcessAllocations );
+
 	// Initialize the begin snapshot.
 	auto BeginSnapshot = SnapshotsToBeProcessed[0];
 	SnapshotsToBeProcessed.RemoveAt( 0 );
@@ -547,14 +592,15 @@ void FStatsMemoryDumpCommand::ProcessMemoryOperations( FStatsReadFile* InFile )
 			const double CurrentSeconds = FPlatformTime::Seconds();
 			if( CurrentSeconds > PreviousSeconds + NumSecondsBetweenLogs )
 			{
-				const int32 CurrentPct = int32( 100.0*(Index + 1) / NumSequenceAllocations );
-				UE_LOG( LogStats, Log, TEXT( "Processing allocations %3i%% (%10i/%10i)" ), CurrentPct, Index + 1, NumSequenceAllocations );
+				const int32 PercentagePos = int32( 100.0*(Index + 1) / NumSequenceAllocations );
+				InFile->StageProgress.Set( PercentagePos );
+				UE_LOG( LogStats, Verbose, TEXT( "Processing allocations %3i%% (%10i/%10i)" ), PercentagePos, Index + 1, NumSequenceAllocations );
 				PreviousSeconds = CurrentSeconds;
 
 				// Abandon support, simply return for now.
 				if (InFile->bShouldStopProcessing == true)
 				{
-					InFile->SetProcessingStage( EStatsProcessingStage::RS_Stopped );
+					InFile->SetProcessingStage( EStatsProcessingStage::SPS_Stopped );
 					return;
 				}
 			}
@@ -675,6 +721,7 @@ void FStatsMemoryDumpCommand::ProcessMemoryOperations( FStatsReadFile* InFile )
 			}
 		}
 	}
+	InFile->StageProgress.Set( 100 );
 
 	auto EndSnapshot = SnapshotsToBeProcessed[0];
 	SnapshotsToBeProcessed.RemoveAt( 0 );
@@ -713,44 +760,7 @@ void FStatsMemoryDumpCommand::ProcessMemoryOperations( FStatsReadFile* InFile )
 		}
 	}
 
-	InFile->SetProcessingStage( EStatsProcessingStage::RS_Finished );
-
-	// Frame-240 Frame-120 Frame-060
-	TMap<FString, FCombinedAllocationInfo> FrameBegin_Exit;
-	CompareSnapshots_FString( TEXT( "BeginSnapshot" ), TEXT( "EngineLoop.Exit" ), FrameBegin_Exit );
-	DumpScopedAllocations( TEXT( "Begin_Exit" ), FrameBegin_Exit );
-
-#if	UE_BUILD_DEBUG
-	TMap<FString, FCombinedAllocationInfo> Frame060_120;
-	CompareSnapshots_FString( TEXT( "Frame-060" ), TEXT( "Frame-120" ), Frame060_120 );
-	DumpScopedAllocations( TEXT( "Frame060_120" ), Frame060_120 );
-
-	TMap<FString, FCombinedAllocationInfo> Frame060_240;
-	CompareSnapshots_FString( TEXT( "Frame-060" ), TEXT( "Frame-240" ), Frame060_240 );
-	DumpScopedAllocations( TEXT( "Frame060_240" ), Frame060_240 );
-
-	// Generate scoped tree view.
-	{
-		TMap<FName, FCombinedAllocationInfo> FrameBegin_Exit_FName;
-		CompareSnapshots_FName( TEXT( "BeginSnapshot" ), TEXT( "EngineLoop.Exit" ), FrameBegin_Exit_FName );
-
-		FNodeAllocationInfo Root;
-		Root.EncodedCallstack = TEXT( "ThreadRoot" );
-		Root.HumanReadableCallstack = TEXT( "ThreadRoot" );
-		GenerateScopedTreeAllocations( FrameBegin_Exit_FName, Root );
-	}
-
-
-	{
-		TMap<FName, FCombinedAllocationInfo> Frame060_240_FName;
-		CompareSnapshots_FName( TEXT( "Frame-060" ), TEXT( "Frame-240" ), Frame060_240_FName );
-
-		FNodeAllocationInfo Root;
-		Root.EncodedCallstack = TEXT( "ThreadRoot" );
-		Root.HumanReadableCallstack = TEXT( "ThreadRoot" );
-		GenerateScopedTreeAllocations( Frame060_240_FName, Root );
-	}
-#endif // UE_BUILD_DEBUG
+	InFile->SetProcessingStage( EStatsProcessingStage::SPS_Finished );
 }
 
 void FStatsMemoryDumpCommand::GenerateScopedTreeAllocations( const TMap<FName, FCombinedAllocationInfo>& ScopedAllocations, FNodeAllocationInfo& out_Root )
diff --git a/Engine/Source/Programs/UnrealFrontend/Private/Commands/StatsDumpMemoryCommand.h b/Engine/Source/Programs/UnrealFrontend/Private/Commands/StatsDumpMemoryCommand.h
--- a/Engine/Source/Programs/UnrealFrontend/Private/Commands/StatsDumpMemoryCommand.h
+++ b/Engine/Source/Programs/UnrealFrontend/Private/Commands/StatsDumpMemoryCommand.h
@@ -206,6 +206,7 @@ protected:
 	}
 };
 
+// #YRX_Profiler: 2015-07-10 Should be moved to the profiler module, and added as export or something?
 /** An example command loading a raw stats file and dumping memory usage. */
 class FStatsMemoryDumpCommand
 {
diff --git a/Engine/Source/Runtime/Core/Private/Stats/StatsFile.cpp b/Engine/Source/Runtime/Core/Private/Stats/StatsFile.cpp
--- a/Engine/Source/Runtime/Core/Private/Stats/StatsFile.cpp
+++ b/Engine/Source/Runtime/Core/Private/Stats/StatsFile.cpp
@@ -584,9 +584,9 @@ void FStatsReadFile::ReadAndProcessSynchronously()
 	ReadStats();
 	ProcessStats();
 
-	if (GetProcessingStage() == EStatsProcessingStage::RS_Stopped)
+	if (GetProcessingStage() == EStatsProcessingStage::SPS_Stopped)
 	{
-		SetProcessingStage( EStatsProcessingStage::RS_Invalid );
+		SetProcessingStage( EStatsProcessingStage::SPS_Invalid );
 	}
 }
 
@@ -616,7 +616,7 @@ void FStatsReadFile::ReadStats()
 
 		// Update stage progress once per NumSecondsBetweenUpdates(5) seconds to avoid spamming.
 		double PreviousSeconds = FPlatformTime::Seconds();
-		SetProcessingStage( EStatsProcessingStage::RS_ReadAndCombinePackets );
+		SetProcessingStage( EStatsProcessingStage::SPS_ReadAndCombinePackets );
 
 		while (Reader->Tell() < Reader->TotalSize())
 		{
@@ -659,14 +659,14 @@ void FStatsReadFile::ReadStats()
 
 			if (bShouldStopProcessing == true)
 			{
-				SetProcessingStage( EStatsProcessingStage::RS_Stopped );
+				SetProcessingStage( EStatsProcessingStage::SPS_Stopped );
 				break;
 			}
 
 			FileInfo.TotalPacketsNum++;
 		}
 
-		if (GetProcessingStage() != EStatsProcessingStage::RS_Stopped)
+		if (GetProcessingStage() != EStatsProcessingStage::SPS_Stopped)
 		{
 			const double TotalTime = FPlatformTime::Seconds() - StartTime;
 			UE_LOG( LogStats, Log, TEXT( "Reading took %.2f sec(s)" ), TotalTime );
@@ -686,12 +686,12 @@ void FStatsReadFile::ProcessStats()
 {
 	if (bRawStatsFile)
 	{
-		if (GetProcessingStage() != EStatsProcessingStage::RS_Stopped)
+		if (GetProcessingStage() != EStatsProcessingStage::SPS_Stopped)
 		{
-			SetProcessingStage( EStatsProcessingStage::RS_ProcessStatMessages );
+			SetProcessingStage( EStatsProcessingStage::SPS_ProcessStatMessages );
 			NewCombinedHistory.ExecuteIfBound( this );
 
-			if (GetProcessingStage() == EStatsProcessingStage::RS_Stopped)
+			if (GetProcessingStage() == EStatsProcessingStage::SPS_Stopped)
 			{
 				UE_LOG( LogStats, Warning, TEXT( "Processing stopped, abandoning" ) );
 			}
@@ -720,7 +720,7 @@ FStatsReadFile::~FStatsReadFile()
 
 	FPlatformProcess::ConditionalSleep( [&]()
 	{
-		return IsBusy();
+		return !IsBusy();
 	}, 1.0f );
 
 	if (AsyncWork)
@@ -739,12 +739,12 @@ bool FStatsReadFile::PrepareLoading()
 {
 	const double StartTime = FPlatformTime::Seconds();
 
-	SetProcessingStage( EStatsProcessingStage::RS_Started );
+	SetProcessingStage( EStatsProcessingStage::SPS_Started );
 
 	{
 		ON_SCOPE_EXIT
 		{
-			SetProcessingStage( EStatsProcessingStage::RS_Invalid );
+			SetProcessingStage( EStatsProcessingStage::SPS_Invalid );
 		};
 
 		const int64 Size = IFileManager::Get().FileSize( *Filename );
diff --git a/Engine/Source/Runtime/Core/Public/Stats/StatsFile.h b/Engine/Source/Runtime/Core/Public/Stats/StatsFile.h
--- a/Engine/Source/Runtime/Core/Public/Stats/StatsFile.h
+++ b/Engine/Source/Runtime/Core/Public/Stats/StatsFile.h
@@ -770,16 +770,18 @@ struct FRawStatsFileInfo
 /** Enumerates stats processing stages. */
 enum EStatsProcessingStage
 {
-	RS_Started = 0,
-	RS_ReadAndCombinePackets,
-	RS_ProcessStatMessages,
+	SPS_Started = 0,
+	SPS_ReadAndCombinePackets,
+	SPS_ProcessStatMessages,
 	/** Memory profiler only. */
-	RS_SortSequence,
-	RS_Finished,
-	RS_Stopped,
+	SPS_SortSequence,
+	/** Memory profiler only. */
+	SPS_ProcessAllocations,
+	SPS_Finished,
+	SPS_Stopped,
 
 	/** Last stage, at this moment all data is read and processed or process has been stopped, we can now remove the instance. */
-	RS_Invalid,
+	SPS_Invalid,
 
 };
 
@@ -865,6 +867,57 @@ public:
 		return EStatsProcessingStage( ProcessingStage.GetValue() );
 	}
 
+	/**
+	 * @return current processing stage as string key.
+	 */
+	FString GetProcessingStageAsString() const
+	{
+		const EStatsProcessingStage Stage = GetProcessingStage();
+		FString Result;
+		if (Stage== EStatsProcessingStage::SPS_Started)
+		{
+			Result = TEXT( "SPS_Started" );
+		}
+		else if (Stage == EStatsProcessingStage::SPS_ReadAndCombinePackets)
+		{
+			Result = TEXT( "SPS_ReadAndCombinePackets" );
+		}
+		else if (Stage == EStatsProcessingStage::SPS_ProcessStatMessages)
+		{
+			Result = TEXT( "SPS_ProcessStatMessages" );
+		}
+		else if (Stage == EStatsProcessingStage::SPS_SortSequence)
+		{
+			Result = TEXT( "SPS_SortSequence" );
+		}
+		else if (Stage == EStatsProcessingStage::SPS_ProcessAllocations)
+		{
+			Result = TEXT( "SPS_ProcessAllocations" );
+		}
+		else if (Stage == EStatsProcessingStage::SPS_Finished)
+		{
+			Result = TEXT( "SPS_Finished" );
+		}
+		else if (Stage == EStatsProcessingStage::SPS_Stopped)
+		{
+			Result = TEXT( "SPS_Stopped" );
+		}
+		else if (Stage == EStatsProcessingStage::SPS_Invalid)
+		{
+			Result = TEXT( "SPS_Invalid" );
+		}
+
+		return Result;
+	}
+
+	/**
+	 * @return current stage progress as percentage, between 0 and 100.
+	 */
+	int32 GetStageProgress() const
+	{
+		return StageProgress.GetValue();
+	}
+
 	/**	 
 	 * @return true, if any asynchronous operation is being processed
 	 */
@@ -878,14 +931,6 @@ public:
 	{
 		bShouldStopProcessing = true;
 	}
-	
-	/**
-	 * @return current stage of reading.
-	 */
-	const EStatsProcessingStage GetReadingStage() const
-	{
-		return (EStatsProcessingStage)ProcessingStage.GetValue();
-	}
 
 protected:
 	/** Initialization constructor. */
