Stats/Profiler - Refactored reading from a stats file (move memorydumpcommand to the profiler, all memory profiler funcionality move to the profiler module, can be accessed from the IProfilerModule, still wip)

diff --git a/Engine/Source/Developer/Profiler/Private/ProfilerManager.cpp b/Engine/Source/Developer/Profiler/Private/ProfilerManager.cpp
--- a/Engine/Source/Developer/Profiler/Private/ProfilerManager.cpp
+++ b/Engine/Source/Developer/Profiler/Private/ProfilerManager.cpp
@@ -27,7 +27,6 @@ DEFINE_STAT(STAT_PM_MemoryUsage);
 -----------------------------------------------------------------------------*/
 
 TSharedPtr<FProfilerManager> FProfilerManager::Instance = nullptr;
-FThreadSafeCounter FProfilerManager::ProcessingLock;
 
 struct FEventGraphSampleLess
 {
diff --git a/Engine/Source/Developer/Profiler/Private/ProfilerManager.h b/Engine/Source/Developer/Profiler/Private/ProfilerManager.h
--- a/Engine/Source/Developer/Profiler/Private/ProfilerManager.h
+++ b/Engine/Source/Developer/Profiler/Private/ProfilerManager.h
@@ -163,19 +163,12 @@ namespace EProfilerViewMode
  ** */
 class FProfilerManager 
 	: public TSharedFromThis<FProfilerManager>
-	, public IProfilerManager
 {
 	friend class FProfilerActionManager;
 
-	/**
-	 *	Global processing lock, protects from destroying the profiler while task graph is still working on the stuff.
-	 *	Assumes that there is only one instance of the profiler.
-	 */
-	static FThreadSafeCounter ProcessingLock;
-
 public:
 	/**
-	 * Creates a profiler manager, only one distance can exists.
+	 * Creates a profiler manager, only one instance can exist.
 	 *
 	 * @param InSessionManager	- the session manager to use
 	 *
@@ -193,7 +186,7 @@ public:
 	 */
 	static TSharedPtr<FProfilerManager> Initialize( const ISessionManagerRef& SessionManager )
 	{
-		if ( FProfilerManager::Instance.IsValid() )
+		if (FProfilerManager::Instance.IsValid())
 		{
 			FProfilerManager::Instance.Reset();
 		}
@@ -213,18 +206,6 @@ public:
 		FProfilerManager::Instance.Reset();
 	}
 
-	/** Increases processing lock count. */
-	static void IncrementProcessingLock()
-	{
-		ProcessingLock.Increment();
-	}
-
-	/** Decreases processing lock count, once reaches zero the profiler can proceed. */
-	static void DecrementProcessingLock()
-	{
-		ProcessingLock.Decrement();
-	}
-
 protected:
 	/**
 	 * Finishes initialization of the profiler manager
diff --git a/Engine/Source/Developer/Profiler/Private/ProfilerModule.cpp b/Engine/Source/Developer/Profiler/Private/ProfilerModule.cpp
--- a/Engine/Source/Developer/Profiler/Private/ProfilerModule.cpp
+++ b/Engine/Source/Developer/Profiler/Private/ProfilerModule.cpp
@@ -1,9 +1,5 @@
 // Copyright 1998-2015 Epic Games, Inc. All Rights Reserved.
 
-/*=============================================================================
-	ProfilerModule.cpp: Implements the FProfilerModule class.
-=============================================================================*/
-
 #include "ProfilerPrivatePCH.h"
 #include "SDockTab.h"
 
@@ -14,48 +10,79 @@ class FProfilerModule
 	: public IProfilerModule
 {
 public:
-	virtual TSharedRef<SWidget> CreateProfilerWindow( const ISessionManagerRef& InSessionManager, const TSharedRef<SDockTab>& ConstructUnderMajorTab ) override
-	{
-		ProfilerManager = FProfilerManager::Initialize( InSessionManager );
-		TSharedRef<SProfilerWindow> ProfilerWindow = SNew(SProfilerWindow);
-		FProfilerManager::Get()->AssignProfilerWindow( ProfilerWindow );
-		// Register OnTabClosed to handle profiler manager shutdown.
-		ConstructUnderMajorTab->SetOnTabClosed( SDockTab::FOnTabClosedCallback::CreateRaw(this, &FProfilerModule::Shutdown) );
-
-		return ProfilerWindow;
-	}
+	virtual TSharedRef<SWidget> CreateProfilerWindow( const ISessionManagerRef& InSessionManager, const TSharedRef<SDockTab>& ConstructUnderMajorTab ) override;
 
 	virtual void StartupModule() override
 	{}
 
-	virtual void ShutdownModule() override
-	{
-		if( FProfilerManager::Get().IsValid() )
-		{
-			FProfilerManager::Get()->Shutdown();
-		}
-	}
+	virtual void ShutdownModule() override;
 
 	virtual bool SupportsDynamicReloading() override
 	{
 		return false;
 	}
 
-	virtual TWeakPtr<class IProfilerManager> GetProfilerManager() override
-	{
-		return ProfilerManager;
-	}
+	void StatsMemoryDumpCommand( const TCHAR* Filename );
+
+	FRawStatsMemoryProfiler* OpenRawStatForMemoryProfiling( const TCHAR* Filename );
+
 
 protected:
 	/** Shutdowns the profiler manager. */
-	void Shutdown( TSharedRef<SDockTab> TabBeingClosed )
+	void Shutdown( TSharedRef<SDockTab> TabBeingClosed );
+};
+
+IMPLEMENT_MODULE( FProfilerModule, Profiler );
+
+/*-----------------------------------------------------------------------------
+	FProfilerModule
+-----------------------------------------------------------------------------*/
+
+
+TSharedRef<SWidget> FProfilerModule::CreateProfilerWindow( const ISessionManagerRef& InSessionManager, const TSharedRef<SDockTab>& ConstructUnderMajorTab )
+{
+	FProfilerManager::Initialize( InSessionManager );
+	TSharedRef<SProfilerWindow> ProfilerWindow = SNew( SProfilerWindow );
+	FProfilerManager::Get()->AssignProfilerWindow( ProfilerWindow );
+	// Register OnTabClosed to handle profiler manager shutdown.
+	ConstructUnderMajorTab->SetOnTabClosed( SDockTab::FOnTabClosedCallback::CreateRaw( this, &FProfilerModule::Shutdown ) );
+
+	return ProfilerWindow;
+}
+
+void FProfilerModule::ShutdownModule()
+{
+	if (FProfilerManager::Get().IsValid())
 	{
 		FProfilerManager::Get()->Shutdown();
-		TabBeingClosed->SetOnTabClosed( SDockTab::FOnTabClosedCallback() );
 	}
+}
+
+void FProfilerModule::Shutdown( TSharedRef<SDockTab> TabBeingClosed )
+{
+	FProfilerManager::Get()->Shutdown();
+	TabBeingClosed->SetOnTabClosed( SDockTab::FOnTabClosedCallback() );
+}
+
+/*-----------------------------------------------------------------------------
+	StatsMemoryDumpCommand
+-----------------------------------------------------------------------------*/
+
+void FProfilerModule::StatsMemoryDumpCommand( const TCHAR* Filename )
+{
+
+}
+
+/*-----------------------------------------------------------------------------
+	FRawStatsMemoryProfiler
+-----------------------------------------------------------------------------*/
+
+FRawStatsMemoryProfiler* FProfilerModule::OpenRawStatForMemoryProfiling( const TCHAR* Filename )
+{
+	FRawStatsMemoryProfiler* Instance = new FRawStatsMemoryProfiler;
+
+	return Instance;
+}
+
 
-	/** A weak pointer to the global instance of the profiler manager. Will last as long as profiler window is visible. */
-	TWeakPtr<class IProfilerManager> ProfilerManager;
-};
 
-IMPLEMENT_MODULE(FProfilerModule, Profiler);
\ No newline at end of file
diff --git a/Engine/Source/Developer/Profiler/Public/Interfaces/IProfilerManager.h b/Engine/Source/Developer/Profiler/Public/Interfaces/IProfilerManager.h
deleted file mode 100644
diff --git a/Engine/Source/Developer/Profiler/Public/Interfaces/IProfilerModule.h b/Engine/Source/Developer/Profiler/Public/Interfaces/IProfilerModule.h
--- a/Engine/Source/Developer/Profiler/Public/Interfaces/IProfilerModule.h
+++ b/Engine/Source/Developer/Profiler/Public/Interfaces/IProfilerModule.h
@@ -2,11 +2,7 @@
 
 #pragma once
 
-
-// forward declarations
-class IProfilerManager;
-class ISessionManager;
-
+class FRawStatsMemoryProfiler;
 
 /**
  * Interface for the profiler module.
@@ -25,8 +21,9 @@ public:
 	 */
 	virtual TSharedRef<SWidget> CreateProfilerWindow( const TSharedRef<ISessionManager>& InSessionManager, const TSharedRef<SDockTab>& ConstructUnderMajorTab ) = 0;
 
-	/**
-	 * @return a weak pointer to the global instance of the profiler manager.
-	 */
-	virtual TWeakPtr<IProfilerManager> GetProfilerManager() = 0;
+	/** Implements stats memory dump command. */
+	virtual void StatsMemoryDumpCommand( const TCHAR* Filename ) = 0;
+
+	/** Creates a new instance of the memory profiler based the raw stats. */
+	virtual FRawStatsMemoryProfiler* OpenRawStatForMemoryProfiling( const TCHAR* Filename ) = 0;
 };
\ No newline at end of file
diff --git a/Engine/Source/Developer/Profiler/Public/Profiler.h b/Engine/Source/Developer/Profiler/Public/Profiler.h
--- a/Engine/Source/Developer/Profiler/Public/Profiler.h
+++ b/Engine/Source/Developer/Profiler/Public/Profiler.h
@@ -30,5 +30,4 @@
 /* Interfaces
  *****************************************************************************/
 
-#include "IProfilerManager.h"
 #include "IProfilerModule.h"
diff --git a/Engine/Source/Programs/UnrealFrontend/Private/Commands/StatsDumpMemoryCommand.cpp b/Engine/Source/Programs/UnrealFrontend/Private/Commands/StatsDumpMemoryCommand.cpp
--- a/Engine/Source/Programs/UnrealFrontend/Private/Commands/StatsDumpMemoryCommand.cpp
+++ b/Engine/Source/Programs/UnrealFrontend/Private/Commands/StatsDumpMemoryCommand.cpp
@@ -1,6 +1,7 @@
 // Copyright 1998-2015 Epic Games, Inc. All Rights Reserved.
 
 #include "UnrealFrontendPrivatePCH.h"
+#include "Profiler.h"
 #include "StatsDumpMemoryCommand.h"
 #include "DiagnosticTable.h"
 
@@ -221,6 +222,21 @@ void FStatsMemoryDumpCommand::InternalRun()
 	FString SourceFilepath;
 	FParse::Value( FCommandLine::Get(), TEXT( "-INFILE=" ), SourceFilepath );
 
+	{
+		const FName NAME_ProfilerModule = TEXT( "Profiler" );
+		IProfilerModule& ProfilerModule = FModuleManager::LoadModuleChecked<IProfilerModule>( NAME_ProfilerModule );
+		ProfilerModule.StatsMemoryDumpCommand( *SourceFilepath );
+		//FModuleManager::Get().UnloadModule( NAME_ProfilerModule );
+	}
+
+	{
+		const FName NAME_ProfilerModule = TEXT( "Profiler" );
+		IProfilerModule& ProfilerModule = FModuleManager::LoadModuleChecked<IProfilerModule>( NAME_ProfilerModule );
+		ProfilerModule.StatsMemoryDumpCommand( *SourceFilepath );
+		//FModuleManager::Get().UnloadModule( NAME_ProfilerModule );
+	}
+
+
 	FStatsReadFile* NewReader = FStatsReadFile::CreateReaderForRawStats( *SourceFilepath, FStatsReadFile::FOnNewCombinedHistory::CreateRaw( this, &FStatsMemoryDumpCommand::ProcessMemoryOperations ) );
 	if (NewReader)
 	{
@@ -231,9 +247,11 @@ void FStatsMemoryDumpCommand::InternalRun()
 		{
 			FPlatformProcess::Sleep( 1.0f );
 
-			UE_LOG( LogStats, Log, TEXT( "Async: Stage: %s/%3i%%" ), *NewReader->GetProcessingStageAsString(), NewReader->GetStageProgress() );
+			UE_LOG( LogStats, Log, TEXT( "Async: Stage: %s / %3i%%" ), *NewReader->GetProcessingStageAsString(), NewReader->GetStageProgress() );
 		}
+		//NewReader->RequestStop();
 
+		// Delete reader, we don't need the data anymore.
 		delete NewReader;
 
 		// Frame-240 Frame-120 Frame-060
@@ -262,15 +280,15 @@ void FStatsMemoryDumpCommand::InternalRun()
 		}
 
 
-	{
-		TMap<FName, FCombinedAllocationInfo> Frame060_240_FName;
-		CompareSnapshots_FName( TEXT( "Frame-060" ), TEXT( "Frame-240" ), Frame060_240_FName );
+		{
+			TMap<FName, FCombinedAllocationInfo> Frame060_240_FName;
+			CompareSnapshots_FName( TEXT( "Frame-060" ), TEXT( "Frame-240" ), Frame060_240_FName );
 
-		FNodeAllocationInfo Root;
-		Root.EncodedCallstack = TEXT( "ThreadRoot" );
-		Root.HumanReadableCallstack = TEXT( "ThreadRoot" );
-		GenerateScopedTreeAllocations( Frame060_240_FName, Root );
-	}
+			FNodeAllocationInfo Root;
+			Root.EncodedCallstack = TEXT( "ThreadRoot" );
+			Root.HumanReadableCallstack = TEXT( "ThreadRoot" );
+			GenerateScopedTreeAllocations( Frame060_240_FName, Root );
+		}
 #endif // UE_BUILD_DEBUG
 	}
 }
@@ -1099,13 +1117,21 @@ void FStatsMemoryDumpCommand::PrepareSnapshot( const FName SnapshotName, const T
 {
 	FScopeLogTime SLT( TEXT( "PrepareSnapshot" ), nullptr, FScopeLogTime::ScopeLog_Milliseconds );
 
-	SnapshotsWithAllocationMap.Add( SnapshotName, AllocationMap );
+	// Make sure the snapshot name is unique.
+	FName UniqueSnapshotName = SnapshotName;
+	while (SnapshotNames.Contains( UniqueSnapshotName ))
+	{
+		UniqueSnapshotName = FName( UniqueSnapshotName, UniqueSnapshotName.GetNumber() + 1 );
+	}
+	SnapshotNames.Add( UniqueSnapshotName );
+
+	SnapshotsWithAllocationMap.Add( UniqueSnapshotName, AllocationMap );
 
 	TMap<FName, FCombinedAllocationInfo> SnapshotCombinedAllocations;
 	uint64 TotalAllocatedMemory = 0;
 	uint64 NumAllocations = 0;
 	GenerateScopedAllocations( AllocationMap, SnapshotCombinedAllocations, TotalAllocatedMemory, NumAllocations );
-	SnapshotsWithScopedAllocations.Add( SnapshotName, SnapshotCombinedAllocations );
+	SnapshotsWithScopedAllocations.Add( UniqueSnapshotName, SnapshotCombinedAllocations );
 
 	// Decode callstacks.
 	// Replace encoded callstacks with human readable name. For easier debugging.
@@ -1115,9 +1141,9 @@ void FStatsMemoryDumpCommand::PrepareSnapshot( const FName SnapshotName, const T
 		const FString HumanReadableCallstack = FStatsCallstack::GetHumanReadable( It.Key );
 		SnapshotDecodedCombinedAllocations.Add( HumanReadableCallstack, It.Value );
 	}
-	SnapshotsWithDecodedScopedAllocations.Add( SnapshotName, SnapshotDecodedCombinedAllocations );
+	SnapshotsWithDecodedScopedAllocations.Add( UniqueSnapshotName, SnapshotDecodedCombinedAllocations );
 
-	UE_LOG( LogStats, Warning, TEXT( "PrepareSnapshot: %s Alloc: %i Scoped: %i Total: %.2f MB" ), *SnapshotName.GetPlainNameString(), AllocationMap.Num(), SnapshotCombinedAllocations.Num(), TotalAllocatedMemory / 1024.0f / 1024.0f );
+	UE_LOG( LogStats, Warning, TEXT( "PrepareSnapshot: %s Alloc: %i Scoped: %i Total: %.2f MB" ), *UniqueSnapshotName.ToString(), AllocationMap.Num(), SnapshotCombinedAllocations.Num(), TotalAllocatedMemory / 1024.0f / 1024.0f );
 }
 
 void FStatsMemoryDumpCommand::CompareSnapshots_FName( const FName BeginSnaphotName, const FName EndSnaphotName, TMap<FName, FCombinedAllocationInfo>& out_Result )
diff --git a/Engine/Source/Programs/UnrealFrontend/Private/Commands/StatsDumpMemoryCommand.h b/Engine/Source/Programs/UnrealFrontend/Private/Commands/StatsDumpMemoryCommand.h
--- a/Engine/Source/Programs/UnrealFrontend/Private/Commands/StatsDumpMemoryCommand.h
+++ b/Engine/Source/Programs/UnrealFrontend/Private/Commands/StatsDumpMemoryCommand.h
@@ -260,6 +260,9 @@ protected:
 	/** The sequence tag mapping to the named markers. */
 	TArray<TPair<uint32, FName>> Snapshots;
 
+	/** Unique snapshot names. */
+	TSet<FName> SnapshotNames;
+
 	/** The sequence tag mapping to the named markers that need to processed. */
 	TArray<TPair<uint32, FName>> SnapshotsToBeProcessed;
 
diff --git a/Engine/Source/Programs/UnrealFrontend/UnrealFrontend.Build.cs b/Engine/Source/Programs/UnrealFrontend/UnrealFrontend.Build.cs
--- a/Engine/Source/Programs/UnrealFrontend/UnrealFrontend.Build.cs
+++ b/Engine/Source/Programs/UnrealFrontend/UnrealFrontend.Build.cs
@@ -24,6 +24,7 @@ public class UnrealFrontend : ModuleRules
 				"DeviceManager",
 				"LauncherServices",
 				"Messaging",
+				"Profiler",
 				"ProfilerClient",
                 "ProjectLauncher",
 				"Projects",
diff --git a/Engine/Source/Runtime/Core/Public/Stats/StatsFile.h b/Engine/Source/Runtime/Core/Public/Stats/StatsFile.h
--- a/Engine/Source/Runtime/Core/Public/Stats/StatsFile.h
+++ b/Engine/Source/Runtime/Core/Public/Stats/StatsFile.h
@@ -768,7 +768,7 @@ struct FRawStatsFileInfo
 };
 
 /** Enumerates stats processing stages. */
-enum EStatsProcessingStage
+enum class EStatsProcessingStage : int32
 {
 	SPS_Started = 0,
 	SPS_ReadAndCombinePackets,
@@ -926,7 +926,7 @@ public:
 		return AsyncWork != nullptr ? !AsyncWork->IsDone() : false;
 	}
 
-	/** Requests stopping the processing of the data. */
+	/** Requests stopping the processing of the data. May take a few seconds to finish. */
 	void RequestStop()
 	{
 		bShouldStopProcessing = true;
