Stats/Profiler - Refactored reading from a stats file (move memorydumpcommand to the profiler, all memory profiler funcionality move to the profiler module, can be accessed from the IProfilerModule, still wip)

diff --git a/Engine/Source/Developer/Profiler/Private/ProfilerModule.cpp b/Engine/Source/Developer/Profiler/Private/ProfilerModule.cpp
--- a/Engine/Source/Developer/Profiler/Private/ProfilerModule.cpp
+++ b/Engine/Source/Developer/Profiler/Private/ProfilerModule.cpp
@@ -24,8 +24,7 @@ public:
 
 	void StatsMemoryDumpCommand( const TCHAR* Filename );
 
-	FRawStatsMemoryProfiler* OpenRawStatForMemoryProfiling( const TCHAR* Filename );
-
+	FRawStatsMemoryProfiler* OpenRawStatsForMemoryProfiling( const TCHAR* Filename );
 
 protected:
 	/** Shutdowns the profiler manager. */
@@ -70,16 +69,68 @@ void FProfilerModule::Shutdown( TSharedRef<SDockTab> TabBeingClosed )
 
 void FProfilerModule::StatsMemoryDumpCommand( const TCHAR* Filename )
 {
-
+	//FRawStatsMemoryProfiler* Instance = FCreateStatsReader<FRawStatsMemoryProfiler>::ForRawStats( Filename );
+	TUniquePtr<FRawStatsMemoryProfiler> Instance( FCreateStatsReader<FRawStatsMemoryProfiler>::ForRawStats( Filename ) );
+	if (Instance)
+	{
+		Instance->ReadAndProcessAsynchronously();
+
+		while (Instance->IsBusy())
+		{
+			FPlatformProcess::Sleep( 1.0f );
+
+			UE_LOG( LogStats, Log, TEXT( "Async: Stage: %s / %3i%%" ), *Instance->GetProcessingStageAsString(), Instance->GetStageProgress() );
+		}
+		//NewReader->RequestStop();
+
+		//Instance->
+
+		// Frame-240 Frame-120 Frame-060
+		TMap<FString, FCombinedAllocationInfo> FrameBegin_Exit;
+		Instance->CompareSnapshots_FString( TEXT( "BeginSnapshot" ), TEXT( "EngineLoop.Exit" ), FrameBegin_Exit );
+		Instance->DumpScopedAllocations( TEXT( "Begin_Exit" ), FrameBegin_Exit );
+
+#if	UE_BUILD_DEBUG
+		TMap<FString, FCombinedAllocationInfo> Frame060_120;
+		Instance->CompareSnapshots_FString( TEXT( "Frame-060" ), TEXT( "Frame-120" ), Frame060_120 );
+		Instance->DumpScopedAllocations( TEXT( "Frame060_120" ), Frame060_120 );
+
+		TMap<FString, FCombinedAllocationInfo> Frame060_240;
+		Instance->CompareSnapshots_FString( TEXT( "Frame-060" ), TEXT( "Frame-240" ), Frame060_240 );
+		Instance->DumpScopedAllocations( TEXT( "Frame060_240" ), Frame060_240 );
+
+		// Generate scoped tree view.
+		{
+			TMap<FName, FCombinedAllocationInfo> FrameBegin_Exit_FName;
+			Instance->CompareSnapshots_FName( TEXT( "BeginSnapshot" ), TEXT( "EngineLoop.Exit" ), FrameBegin_Exit_FName );
+
+			FNodeAllocationInfo Root;
+			Root.EncodedCallstack = TEXT( "ThreadRoot" );
+			Root.HumanReadableCallstack = TEXT( "ThreadRoot" );
+			Instance->GenerateScopedTreeAllocations( FrameBegin_Exit_FName, Root );
+		}
+
+
+		{
+			TMap<FName, FCombinedAllocationInfo> Frame060_240_FName;
+			Instance->CompareSnapshots_FName( TEXT( "Frame-060" ), TEXT( "Frame-240" ), Frame060_240_FName );
+
+			FNodeAllocationInfo Root;
+			Root.EncodedCallstack = TEXT( "ThreadRoot" );
+			Root.HumanReadableCallstack = TEXT( "ThreadRoot" );
+			Instance->GenerateScopedTreeAllocations( Frame060_240_FName, Root );
+		}
+#endif // UE_BUILD_DEBUG
+	}
 }
 
 /*-----------------------------------------------------------------------------
 	FRawStatsMemoryProfiler
 -----------------------------------------------------------------------------*/
 
-FRawStatsMemoryProfiler* FProfilerModule::OpenRawStatForMemoryProfiling( const TCHAR* Filename )
+FRawStatsMemoryProfiler* FProfilerModule::OpenRawStatsForMemoryProfiling( const TCHAR* Filename )
 {
-	FRawStatsMemoryProfiler* Instance = new FRawStatsMemoryProfiler;
+	FRawStatsMemoryProfiler* Instance = 0;// new FRawStatsMemoryProfiler;
 
 	return Instance;
 }
diff --git a/Engine/Source/Developer/Profiler/Private/ProfilerRawStats.cpp b/Engine/Source/Developer/Profiler/Private/ProfilerRawStats.cpp
--- a/Engine/Source/Developer/Profiler/Private/ProfilerRawStats.cpp
+++ b/Engine/Source/Developer/Profiler/Private/ProfilerRawStats.cpp
@@ -1,7 +1,7 @@
 // Copyright 1998-2015 Epic Games, Inc. All Rights Reserved.
 
 #include "ProfilerPrivatePCH.h"
-
+#include "DiagnosticTable.h"
 
 /*-----------------------------------------------------------------------------
 	Sort helpers
@@ -186,23 +186,777 @@ void FNodeAllocationInfo::PrepareCallstackData( const TArray<FName>& InDecodedCa
 }
 
 /*-----------------------------------------------------------------------------
-	Stats stack helpers
+	FRawStatsMemoryProfiler
 -----------------------------------------------------------------------------*/
 
-/** Holds stats stack state, used to preserve continuity when the game frame has changed. */
-struct FStackState
+FRawStatsMemoryProfiler::FRawStatsMemoryProfiler( const TCHAR* InFilename )
+	: FStatsReadFile( InFilename, true )
+	, NumMemoryOperations( 0 )
+	, LastSequenceTagForNamedMarker( 0 )
+{}
+
+void FRawStatsMemoryProfiler::PreProcessStats()
+{
+	Super::PreProcessStats();
+
+	// Begin marker.
+	Snapshots.Add( TPairInitializer<uint32, FName>( LastSequenceTagForNamedMarker, TEXT( "BeginSnapshot" ) ) );
+}
+
+
+void FRawStatsMemoryProfiler::PostProcessStats()
 {
-	/** Default constructor. */
-	FStackState()
-		: bIsBrokenCallstack( false )
-	{}
+	Super::PostProcessStats();
+
+	if (!SortSequenceAllocations())
+	{
+		return;
+	}
+
+	// End marker.
+	Snapshots.Add( TPairInitializer<uint32, FName>( TNumericLimits<uint32>::Max(), TEXT( "EndSnapshot" ) ) );
+
+	// Copy snapshots.
+	SnapshotsToBeProcessed = Snapshots;
+
+	UE_LOG( LogStats, Warning, TEXT( "NumMemoryOperations:   %llu" ), NumMemoryOperations );
+	UE_LOG( LogStats, Warning, TEXT( "SequenceAllocationNum: %i" ), SequenceAllocationArray.Num() );
+
+	// Pass 2.
+	/*
+	TMap<uint32, FAllocationInfo> UniqueSeq;
+	TMultiMap<uint32, FAllocationInfo> OriginalAllocs;
+	TMultiMap<uint32, FAllocationInfo> BrokenAllocs;
+	for (const FAllocationInfo& Alloc : SequenceAllocationArray)
+	{
+		const FAllocationInfo* Found = UniqueSeq.Find( Alloc.SequenceTag );
+		if (!Found)
+		{
+			UniqueSeq.Add( Alloc.SequenceTag, Alloc );
+		}
+		else
+		{
+			OriginalAllocs.Add( Alloc.SequenceTag, *Found );
+			BrokenAllocs.Add( Alloc.SequenceTag, Alloc );
+		}
+	}
+	//*/
+
+	// Alive allocations.
+	TMap<uint64, FAllocationInfo> AllocationMap;
+	TMultiMap<uint64, FAllocationInfo> FreeWithoutAllocMap;
+	TMultiMap<uint64, FAllocationInfo> DuplicatedAllocMap;
+	int32 NumDuplicatedMemoryOperations = 0;
+	int32 NumFWAMemoryOperations = 0; // FreeWithoutAlloc
+	int32 NumZeroAllocs = 0; // Malloc(0)
+
+	// Initialize the begin snapshot.
+	auto BeginSnapshot = SnapshotsToBeProcessed[0];
+	SnapshotsToBeProcessed.RemoveAt( 0 );
+	PrepareSnapshot( BeginSnapshot.Value, AllocationMap );
+
+	auto CurrentSnapshot = SnapshotsToBeProcessed[0];
+
+	UE_LOG( LogStats, Warning, TEXT( "Generating memory operations map" ) );
+	const int32 NumSequenceAllocations = SequenceAllocationArray.Num();
+	const int32 OnePercent = FMath::Max( NumSequenceAllocations / 100, 1024 );
+	for (int32 Index = 0; Index < NumSequenceAllocations; Index++)
+	{
+		if (Index % OnePercent)
+		{
+			const double CurrentSeconds = FPlatformTime::Seconds();
+			if (CurrentSeconds > LastUpdateTime + NumSecondsBetweenUpdates)
+			{
+				const int32 PercentagePos = int32( 100.0*(Index + 1) / NumSequenceAllocations );
+				StageProgress.Set( PercentagePos );
+				UE_LOG( LogStats, Verbose, TEXT( "Processing allocations %3i%% (%10i/%10i)" ), PercentagePos, Index + 1, NumSequenceAllocations );
+				LastUpdateTime = CurrentSeconds;
+
+				// Abandon support, simply return for now.
+				if (bShouldStopProcessing == true)
+				{
+					SetProcessingStage( EStatsProcessingStage::SPS_Stopped );
+					return;
+				}
+			}
+		}
 
-	/** Call stack. */
-	TArray<FName> Stack;
+		const FAllocationInfo& Alloc = SequenceAllocationArray[Index];
 
-	/** Current function name. */
-	FName Current;
+		// Check named marker/snapshots
+		if (Alloc.SequenceTag > CurrentSnapshot.Key)
+		{
+			SnapshotsToBeProcessed.RemoveAt( 0 );
+			PrepareSnapshot( CurrentSnapshot.Value, AllocationMap );
+			CurrentSnapshot = SnapshotsToBeProcessed[0];
+		}
+
+		if (Alloc.Op == EMemoryOperation::Alloc)
+		{
+			if (Alloc.Size == 0)
+			{
+				NumZeroAllocs++;
+			}
+
+			const FAllocationInfo* Found = AllocationMap.Find( Alloc.Ptr );
+
+			if (!Found)
+			{
+				AllocationMap.Add( Alloc.Ptr, Alloc );
+			}
+			else
+			{
+				const FAllocationInfo* FoundAndFreed = FreeWithoutAllocMap.Find( Found->Ptr );
+				const FAllocationInfo* FoundAndAllocated = FreeWithoutAllocMap.Find( Alloc.Ptr );
+
+#if	UE_BUILD_DEBUG
+				if (FoundAndFreed)
+				{
+					const FString FoundAndFreedCallstack = FStatsCallstack::GetHumanReadable( FoundAndFreed->EncodedCallstack );
+				}
+
+				if (FoundAndAllocated)
+				{
+					const FString FoundAndAllocatedCallstack = FStatsCallstack::GetHumanReadable( FoundAndAllocated->EncodedCallstack );
+				}
+
+				NumDuplicatedMemoryOperations++;
+
+
+				const FString FoundCallstack = FStatsCallstack::GetHumanReadable( Found->EncodedCallstack );
+				const FString AllocCallstack = FStatsCallstack::GetHumanReadable( Alloc.EncodedCallstack );
+				UE_LOG( LogStats, Warning, TEXT( "DuplicateAlloc: %s Size: %i/%i Ptr: %i/%i Tag: %i/%i" ), *AllocCallstack, Found->Size, Alloc.Size, Found->Ptr, Alloc.Ptr, Found->SequenceTag, Alloc.SequenceTag );
+#endif // UE_BUILD_DEBUG
+
+				// Replace pointer.
+				AllocationMap.Add( Alloc.Ptr, Alloc );
+				// Store the old pointer.
+				DuplicatedAllocMap.Add( Alloc.Ptr, *Found );
+			}
+		}
+		else if (Alloc.Op == EMemoryOperation::Realloc)
+		{
+			if (Alloc.Size == 0)
+			{
+				NumZeroAllocs++;
+			}
+
+			// Previous Alloc or Realloc
+			const FAllocationInfo* FoundOld = AllocationMap.Find( Alloc.OldPtr );
+
+			if (FoundOld)
+			{
+				const bool bIsValid = Alloc.SequenceTag > FoundOld->SequenceTag;
+				if (!bIsValid)
+				{
+					UE_LOG( LogStats, Warning, TEXT( "InvalidRealloc Ptr: %llu, Seq: %i/%i" ), Alloc.Ptr, Alloc.SequenceTag, FoundOld->SequenceTag );
+				}
+				// Remove the old allocation.
+				AllocationMap.Remove( Alloc.OldPtr );
+				AllocationMap.Add( Alloc.Ptr, Alloc );
+			}
+			// If we have situation where the old pointer is the same as the new pointer
+			// There is high change that the original call looked like this 
+			// Ptr = Malloc( 40 ); Ptr = Realloc( Ptr, 40 );
+			else if (Alloc.OldPtr != Alloc.Ptr)
+			{
+				// No OldPtr, should not happen as it means realloc without initial alloc.
+				// Or Realloc after Malloc(0)
+
+#if	UE_BUILD_DEBUG
+				const FString ReallocCallstack = FStatsCallstack::GetHumanReadable( Alloc.EncodedCallstack );
+				//UE_LOG( LogStats, Warning, TEXT( "ReallocWithoutAlloc: %s %i %i/%i [%i]" ), *ReallocCallstack, Alloc.Size, Alloc.OldPtr, Alloc.Ptr, Alloc.SequenceTag );
+#endif // UE_BUILD_DEBUG
+
+				AllocationMap.Add( Alloc.Ptr, Alloc );
+			}
+
+		}
+		else if (Alloc.Op == EMemoryOperation::Free)
+		{
+			const FAllocationInfo* Found = AllocationMap.Find( Alloc.Ptr );
+			if (Found)
+			{
+				const bool bIsValid = Alloc.SequenceTag > Found->SequenceTag;
+				if (!bIsValid)
+				{
+					UE_LOG( LogStats, Warning, TEXT( "InvalidFree Ptr: %llu, Seq: %i/%i" ), Alloc.Ptr, Alloc.SequenceTag, Found->SequenceTag );
+				}
+				AllocationMap.Remove( Alloc.Ptr );
+			}
+			else
+			{
+				FreeWithoutAllocMap.Add( Alloc.Ptr, Alloc );
+				NumFWAMemoryOperations++;
+
+#if	UE_BUILD_DEBUG
+				const FString FWACallstack = FStatsCallstack::GetHumanReadable( Alloc.EncodedCallstack );
+				//UE_LOG( LogStats, Warning, TEXT( "FreeWithoutAllocCallstack: %s %i" ), *FWACallstack, Alloc.Ptr );
+#endif // UE_BUILD_DEBUG
+			}
+		}
+	}
+	StageProgress.Set( 100 );
+
+	auto EndSnapshot = SnapshotsToBeProcessed[0];
+	SnapshotsToBeProcessed.RemoveAt( 0 );
+	PrepareSnapshot( EndSnapshot.Value, AllocationMap );
+
+	UE_LOG( LogStats, Warning, TEXT( "NumDuplicatedMemoryOperations: %i" ), NumDuplicatedMemoryOperations );
+	UE_LOG( LogStats, Warning, TEXT( "NumFWAMemoryOperations:        %i" ), NumFWAMemoryOperations );
+	UE_LOG( LogStats, Warning, TEXT( "NumZeroAllocs:                 %i" ), NumZeroAllocs );
+
+	// Dump problematic allocations
+	DuplicatedAllocMap.ValueSort( FAllocationInfoSizeGreater() );
+	//FreeWithoutAllocMap
+
+	uint64 TotalDuplicatedMemory = 0;
+	for (const auto& It : DuplicatedAllocMap)
+	{
+		const FAllocationInfo& Alloc = It.Value;
+		TotalDuplicatedMemory += Alloc.Size;
+	}
+
+	UE_LOG( LogStats, Warning, TEXT( "Dumping duplicated alloc map" ) );
+	const float MaxPctDisplayed = 0.80f;
+	uint64 DisplayedSoFar = 0;
+	for (const auto& It : DuplicatedAllocMap)
+	{
+		const FAllocationInfo& Alloc = It.Value;
+		const FString AllocCallstack = FStatsCallstack::GetHumanReadable( Alloc.EncodedCallstack );
+		UE_LOG( LogStats, Log, TEXT( "%lli (%.2f MB) %s" ), Alloc.Size, Alloc.Size / 1024.0f / 1024.0f, *AllocCallstack );
+
+		DisplayedSoFar += Alloc.Size;
+
+		const float CurrentPct = (float)DisplayedSoFar / (float)TotalDuplicatedMemory;
+		if (CurrentPct > MaxPctDisplayed)
+		{
+			break;
+		}
+	}
+}
+
+void FRawStatsMemoryProfiler::ProcessSpecialMessageMarkerOperation( const FStatMessage& Message, const FStackState& StackState )
+{
+	const FName RawName = Message.NameAndInfo.GetRawName();
+	if (RawName == FStatConstants::NAME_NamedMarker)
+	{
+		const FName NamedMarker = Message.GetValue_FName();
+		Snapshots.Add( TPairInitializer<uint32, FName>( LastSequenceTagForNamedMarker, NamedMarker ) );
+	}
+}
+
+void FRawStatsMemoryProfiler::ProcessMemoryOperation( EMemoryOperation MemOp, uint64 Ptr, uint64 NewPtr, int64 Size, uint32 SequenceTag, const FStackState& StackState )
+{
+	if (MemOp == EMemoryOperation::Alloc)
+	{
+		NumMemoryOperations++;
+	
+		// Add a new allocation.
+		SequenceAllocationArray.Add(
+			FAllocationInfo(
+			0,
+			Ptr,
+			Size,
+			StackState.Stack,
+			SequenceTag,
+			EMemoryOperation::Alloc,
+			StackState.bIsBrokenCallstack
+			) );
+		LastSequenceTagForNamedMarker = SequenceTag;
+
+	}
+	else if (MemOp == EMemoryOperation::Realloc)
+	{
+		NumMemoryOperations++;
+		const uint64 OldPtr = Ptr;
+
+		// Add a new realloc.
+		SequenceAllocationArray.Add(
+			FAllocationInfo(
+			OldPtr,
+			NewPtr,
+			Size,
+			StackState.Stack,
+			SequenceTag,
+			EMemoryOperation::Realloc,
+			StackState.bIsBrokenCallstack
+			) );
+		LastSequenceTagForNamedMarker = SequenceTag;
+	}
+	else if (MemOp == EMemoryOperation::Free)
+	{
+		NumMemoryOperations++;
+
+		// Add a new free.
+		SequenceAllocationArray.Add(
+			FAllocationInfo(
+			0,
+			Ptr,
+			0,
+			StackState.Stack,
+			SequenceTag,
+			EMemoryOperation::Free,
+			StackState.bIsBrokenCallstack
+			) );
+	}
+}
+
+bool FRawStatsMemoryProfiler::SortSequenceAllocations()
+{
+	// Sort all memory operation by the sequence tag, iterate through all operation and generate memory usage.
+	SequenceAllocationArray.Sort( FAllocationInfoSequenceTagLess() );
+
+	// Abandon support.
+	if (bShouldStopProcessing == true)
+	{
+		SetProcessingStage( EStatsProcessingStage::SPS_Stopped );
+		return false;
+	}
+	else
+	{
+		return true;
+	}
+}
+
+void FRawStatsMemoryProfiler::GenerateScopedTreeAllocations( const TMap<FName, FCombinedAllocationInfo>& ScopedAllocations, FNodeAllocationInfo& out_Root )
+{
+	// Experimental code, partially optimized.
+	// Decode all scoped allocations.
+	// Generate tree for allocations and combine them.
+
+	for (const auto& It : ScopedAllocations)
+	{
+		const FName& EncodedCallstack = It.Key;
+		const FCombinedAllocationInfo& CombinedAllocation = It.Value;
+
+		// Decode callstack.
+		TArray<FName> DecodedCallstack;
+		FStatsCallstack::DecodeToNames( EncodedCallstack, DecodedCallstack );
+
+		const int32 AllocationLenght = DecodedCallstack.Num();
+		check( DecodedCallstack.Num() > 0 );
+
+		FNodeAllocationInfo* CurrentNode = &out_Root;
+		// Accumulate with thread root node.
+		CurrentNode->Accumulate( CombinedAllocation );
+
+		// Iterate through the callstack and prepare all nodes if needed, and accumulate memory.
+		TArray<FName> CurrentCallstack;
+		const int32 NumEntries = DecodedCallstack.Num();
+		for (int32 Idx1 = 0; Idx1 < NumEntries; ++Idx1)
+		{
+			const FName NodeName = DecodedCallstack[Idx1];
+			CurrentCallstack.Add( NodeName );
+
+			FNodeAllocationInfo* Node = nullptr;
+			const bool bContainsNode = CurrentNode->ChildNodes.Contains( NodeName );
+			if (!bContainsNode)
+			{
+				Node = new FNodeAllocationInfo;
+				Node->Depth = Idx1;
+				Node->PrepareCallstackData( CurrentCallstack );
 
-	/** Whether this callstack is marked as broken due to mismatched start and end scope cycles. */
-	bool bIsBrokenCallstack;
-};
\ No newline at end of file
+				CurrentNode->ChildNodes.Add( NodeName, Node );
+			}
+			else
+			{
+				Node = CurrentNode->ChildNodes.FindChecked( NodeName );
+			}
+
+			// Accumulate memory usage and num allocations for all nodes in the callstack.
+			Node->Accumulate( CombinedAllocation );
+
+			// Move to the next node.
+			Node->Parent = CurrentNode;
+			CurrentNode = Node;
+		}
+	}
+
+	out_Root.SortBySize();
+}
+
+
+void FRawStatsMemoryProfiler::GenerateMemoryUsageReport( const TMap<uint64, FAllocationInfo>& AllocationMap, const TSet<FName>& UObjectRawNames )
+{
+	if (AllocationMap.Num() == 0)
+	{
+		UE_LOG( LogStats, Warning, TEXT( "There are no allocations, make sure memory profiler is enabled" ) );
+	}
+	else
+	{
+		//ProcessAndDumpUObjectAllocations( AllocationMap );
+		ProcessAndDumpScopedAllocations( AllocationMap );
+	}
+}
+
+void FRawStatsMemoryProfiler::ProcessAndDumpScopedAllocations( const TMap<uint64, FAllocationInfo>& AllocationMap )
+{
+	// This code is not optimized. 
+	FScopeLogTime SLT( TEXT( "ProcessingScopedAllocations" ), nullptr, FScopeLogTime::ScopeLog_Seconds );
+	UE_LOG( LogStats, Warning, TEXT( "Processing scoped allocations" ) );
+
+	const FString ReportName = FString::Printf( TEXT( "%s-Memory-Scoped" ), *GetPlatformName() );
+	FDiagnosticTableViewer MemoryReport( *FDiagnosticTableViewer::GetUniqueTemporaryFilePath( *ReportName ), true );
+
+	// Write a row of headings for the table's columns.
+	MemoryReport.AddColumn( TEXT( "Size (bytes)" ) );
+	MemoryReport.AddColumn( TEXT( "Size (MB)" ) );
+	MemoryReport.AddColumn( TEXT( "Count" ) );
+	MemoryReport.AddColumn( TEXT( "Callstack" ) );
+	MemoryReport.CycleRow();
+
+	TMap<FName, FCombinedAllocationInfo> CombinedAllocations;
+	uint64 TotalAllocatedMemory = 0;
+	uint64 NumAllocations = 0;
+	GenerateScopedAllocations( AllocationMap, CombinedAllocations, TotalAllocatedMemory, NumAllocations );
+
+	// Dump memory to the log.
+	CombinedAllocations.ValueSort( FCombinedAllocationInfoSizeGreater() );
+
+	const float MaxPctDisplayed = 0.90f;
+	int32 CurrentIndex = 0;
+	uint64 DisplayedSoFar = 0;
+	UE_LOG( LogStats, Warning, TEXT( "Index, Size (Size MB), Count, Stat desc" ) );
+	for (const auto& It : CombinedAllocations)
+	{
+		const FCombinedAllocationInfo& CombinedAllocation = It.Value;
+		const FName& EncodedCallstack = It.Key;
+
+		const FString AllocCallstack = FStatsCallstack::GetHumanReadable( EncodedCallstack );
+
+		UE_LOG( LogStats, Log, TEXT( "%2i, %llu (%.2f MB), %llu, %s" ),
+				CurrentIndex,
+				CombinedAllocation.Size,
+				CombinedAllocation.Size / 1024.0f / 1024.0f,
+				CombinedAllocation.Count,
+				*AllocCallstack );
+
+		// Dump stats
+		MemoryReport.AddColumn( TEXT( "%llu" ), CombinedAllocation.Size );
+		MemoryReport.AddColumn( TEXT( "%.2f MB" ), CombinedAllocation.Size / 1024.0f / 1024.0f );
+		MemoryReport.AddColumn( TEXT( "%llu" ), CombinedAllocation.Count );
+		MemoryReport.AddColumn( *AllocCallstack );
+		MemoryReport.CycleRow();
+
+		CurrentIndex++;
+		DisplayedSoFar += CombinedAllocation.Size;
+
+		const float CurrentPct = (float)DisplayedSoFar / (float)TotalAllocatedMemory;
+		if (CurrentPct > MaxPctDisplayed)
+		{
+			break;
+		}
+	}
+
+	UE_LOG( LogStats, Warning, TEXT( "Allocated memory: %llu bytes (%.2f MB)" ), TotalAllocatedMemory, TotalAllocatedMemory / 1024.0f / 1024.0f );
+
+	// Add a total row.
+	MemoryReport.CycleRow();
+	MemoryReport.CycleRow();
+	MemoryReport.CycleRow();
+	MemoryReport.AddColumn( TEXT( "%llu" ), TotalAllocatedMemory );
+	MemoryReport.AddColumn( TEXT( "%.2f MB" ), TotalAllocatedMemory / 1024.0f / 1024.0f );
+	MemoryReport.AddColumn( TEXT( "%llu" ), NumAllocations );
+	MemoryReport.AddColumn( TEXT( "TOTAL" ) );
+	MemoryReport.CycleRow();
+}
+
+void FRawStatsMemoryProfiler::ProcessAndDumpUObjectAllocations( const TMap<uint64, FAllocationInfo>& AllocationMap, const TSet<FName>& UObjectRawNames )
+{
+	// This code is not optimized. 
+	FScopeLogTime SLT( TEXT( "ProcessingUObjectAllocations" ), nullptr, FScopeLogTime::ScopeLog_Seconds );
+	UE_LOG( LogStats, Warning, TEXT( "Processing UObject allocations" ) );
+
+	const FString ReportName = FString::Printf( TEXT( "%s-Memory-UObject" ), *GetPlatformName() );
+	FDiagnosticTableViewer MemoryReport( *FDiagnosticTableViewer::GetUniqueTemporaryFilePath( *ReportName ), true );
+
+	// Write a row of headings for the table's columns.
+	MemoryReport.AddColumn( TEXT( "Size (bytes)" ) );
+	MemoryReport.AddColumn( TEXT( "Size (MB)" ) );
+	MemoryReport.AddColumn( TEXT( "Count" ) );
+	MemoryReport.AddColumn( TEXT( "UObject class" ) );
+	MemoryReport.CycleRow();
+
+	TMap<FName, FCombinedAllocationInfo> UObjectAllocations;
+
+	// To minimize number of calls to expensive DecodeCallstack.
+	TMap<FName, FName> UObjectCallstackToClassMapping;
+
+	uint64 NumAllocations = 0;
+	uint64 TotalAllocatedMemory = 0;
+	for (const auto& It : AllocationMap)
+	{
+		const FAllocationInfo& Alloc = It.Value;
+
+		FName UObjectClass = UObjectCallstackToClassMapping.FindRef( Alloc.EncodedCallstack );
+		if (UObjectClass == NAME_None)
+		{
+			TArray<FName> DecodedCallstack;
+			FStatsCallstack::DecodeToNames( Alloc.EncodedCallstack, DecodedCallstack );
+
+			for (int32 Index = DecodedCallstack.Num() - 1; Index >= 0; --Index)
+			{
+				const FName LongName = DecodedCallstack[Index];
+				const bool bValid = UObjectRawNames.Contains( LongName );
+				if (bValid)
+				{
+					const FString ObjectName = FStatNameAndInfo::GetShortNameFrom( LongName ).GetPlainNameString();
+					UObjectClass = *ObjectName.Left( ObjectName.Find( TEXT( "//" ) ) );;
+					UObjectCallstackToClassMapping.Add( Alloc.EncodedCallstack, UObjectClass );
+					break;
+				}
+			}
+		}
+
+		if (UObjectClass != NAME_None)
+		{
+			FCombinedAllocationInfo& CombinedAllocation = UObjectAllocations.FindOrAdd( UObjectClass );
+			CombinedAllocation += Alloc;
+
+			TotalAllocatedMemory += Alloc.Size;
+			NumAllocations++;
+		}
+	}
+
+	// Dump memory to the log.
+	UObjectAllocations.ValueSort( FCombinedAllocationInfoSizeGreater() );
+
+	const float MaxPctDisplayed = 0.90f;
+	int32 CurrentIndex = 0;
+	uint64 DisplayedSoFar = 0;
+	UE_LOG( LogStats, Warning, TEXT( "Index, Size (Size MB), Count, UObject class" ) );
+	for (const auto& It : UObjectAllocations)
+	{
+		const FCombinedAllocationInfo& CombinedAllocation = It.Value;
+		const FName& UObjectClass = It.Key;
+
+		UE_LOG( LogStats, Log, TEXT( "%2i, %llu (%.2f MB), %llu, %s" ),
+				CurrentIndex,
+				CombinedAllocation.Size,
+				CombinedAllocation.Size / 1024.0f / 1024.0f,
+				CombinedAllocation.Count,
+				*UObjectClass.GetPlainNameString() );
+
+		// Dump stats
+		MemoryReport.AddColumn( TEXT( "%llu" ), CombinedAllocation.Size );
+		MemoryReport.AddColumn( TEXT( "%.2f MB" ), CombinedAllocation.Size / 1024.0f / 1024.0f );
+		MemoryReport.AddColumn( TEXT( "%llu" ), CombinedAllocation.Count );
+		MemoryReport.AddColumn( *UObjectClass.GetPlainNameString() );
+		MemoryReport.CycleRow();
+
+		CurrentIndex++;
+		DisplayedSoFar += CombinedAllocation.Size;
+
+		const float CurrentPct = (float)DisplayedSoFar / (float)TotalAllocatedMemory;
+		if (CurrentPct > MaxPctDisplayed)
+		{
+			break;
+		}
+	}
+
+	UE_LOG( LogStats, Warning, TEXT( "Allocated memory: %llu bytes (%.2f MB)" ), TotalAllocatedMemory, TotalAllocatedMemory / 1024.0f / 1024.0f );
+
+	// Add a total row.
+	MemoryReport.CycleRow();
+	MemoryReport.CycleRow();
+	MemoryReport.CycleRow();
+	MemoryReport.AddColumn( TEXT( "%llu" ), TotalAllocatedMemory );
+	MemoryReport.AddColumn( TEXT( "%.2f MB" ), TotalAllocatedMemory / 1024.0f / 1024.0f );
+	MemoryReport.AddColumn( TEXT( "%llu" ), NumAllocations );
+	MemoryReport.AddColumn( TEXT( "TOTAL" ) );
+	MemoryReport.CycleRow();
+}
+
+void FRawStatsMemoryProfiler::DumpScopedAllocations( const TCHAR* Name, const TMap<FString, FCombinedAllocationInfo>& CombinedAllocations )
+{
+	if (CombinedAllocations.Num() == 0)
+	{
+		UE_LOG( LogStats, Warning, TEXT( "No scoped allocations: %s" ), Name );
+		return;
+	}
+
+	// This code is not optimized. 
+	FScopeLogTime SLT( TEXT( "ProcessingScopedAllocations" ), nullptr, FScopeLogTime::ScopeLog_Seconds );
+	UE_LOG( LogStats, Warning, TEXT( "Dumping scoped allocations: %s" ), Name );
+
+	const FString ReportName = FString::Printf( TEXT( "%s-Memory-Scoped-%s" ), *GetPlatformName(), Name );
+	FDiagnosticTableViewer MemoryReport( *FDiagnosticTableViewer::GetUniqueTemporaryFilePath( *ReportName ), true );
+
+	// Write a row of headings for the table's columns.
+	MemoryReport.AddColumn( TEXT( "Size (bytes)" ) );
+	MemoryReport.AddColumn( TEXT( "Size (MB)" ) );
+	MemoryReport.AddColumn( TEXT( "Count" ) );
+	MemoryReport.AddColumn( TEXT( "Callstack" ) );
+	MemoryReport.CycleRow();
+
+
+	FCombinedAllocationInfo Total;
+
+	const float MaxPctDisplayed = 0.90f;
+	int32 CurrentIndex = 0;
+	UE_LOG( LogStats, Warning, TEXT( "Index, Size (Size MB), Count, Stat desc" ) );
+	for (const auto& It : CombinedAllocations)
+	{
+		const FCombinedAllocationInfo& CombinedAllocation = It.Value;
+		//const FName& EncodedCallstack = It.Key;
+		const FString AllocCallstack = It.Key;// GetCallstack( EncodedCallstack );
+
+		UE_LOG( LogStats, Log, TEXT( "%2i, %llu (%.2f MB), %llu, %s" ),
+				CurrentIndex,
+				CombinedAllocation.Size,
+				CombinedAllocation.Size / 1024.0f / 1024.0f,
+				CombinedAllocation.Count,
+				*AllocCallstack );
+
+		// Dump stats
+		MemoryReport.AddColumn( TEXT( "%llu" ), CombinedAllocation.Size );
+		MemoryReport.AddColumn( TEXT( "%.2f MB" ), CombinedAllocation.Size / 1024.0f / 1024.0f );
+		MemoryReport.AddColumn( TEXT( "%llu" ), CombinedAllocation.Count );
+		MemoryReport.AddColumn( *AllocCallstack );
+		MemoryReport.CycleRow();
+
+		CurrentIndex++;
+		Total += CombinedAllocation;
+	}
+
+	UE_LOG( LogStats, Warning, TEXT( "Allocated memory: %llu bytes (%.2f MB)" ), Total.Size, Total.SizeMB );
+
+	// Add a total row.
+	MemoryReport.CycleRow();
+	MemoryReport.CycleRow();
+	MemoryReport.CycleRow();
+	MemoryReport.AddColumn( TEXT( "%llu" ), Total.Size );
+	MemoryReport.AddColumn( TEXT( "%.2f MB" ), Total.SizeMB );
+	MemoryReport.AddColumn( TEXT( "%llu" ), Total.Count );
+	MemoryReport.AddColumn( TEXT( "TOTAL" ) );
+	MemoryReport.CycleRow();
+}
+
+void FRawStatsMemoryProfiler::GenerateScopedAllocations( const TMap<uint64, FAllocationInfo>& AllocationMap, TMap<FName, FCombinedAllocationInfo>& out_CombinedAllocations, uint64& TotalAllocatedMemory, uint64& NumAllocations )
+{
+	FScopeLogTime SLT( TEXT( "GenerateScopedAllocations" ), nullptr, FScopeLogTime::ScopeLog_Milliseconds );
+
+	for (const auto& It : AllocationMap)
+	{
+		const FAllocationInfo& Alloc = It.Value;
+		FCombinedAllocationInfo& CombinedAllocation = out_CombinedAllocations.FindOrAdd( Alloc.EncodedCallstack );
+		CombinedAllocation += Alloc;
+
+		TotalAllocatedMemory += Alloc.Size;
+		NumAllocations++;
+	}
+
+	// Sort by size.
+	out_CombinedAllocations.ValueSort( FCombinedAllocationInfoSizeGreater() );
+}
+
+void FRawStatsMemoryProfiler::PrepareSnapshot( const FName SnapshotName, const TMap<uint64, FAllocationInfo>& AllocationMap )
+{
+	FScopeLogTime SLT( TEXT( "PrepareSnapshot" ), nullptr, FScopeLogTime::ScopeLog_Milliseconds );
+
+	// Make sure the snapshot name is unique.
+	FName UniqueSnapshotName = SnapshotName;
+	while (SnapshotNames.Contains( UniqueSnapshotName ))
+	{
+		UniqueSnapshotName = FName( UniqueSnapshotName, UniqueSnapshotName.GetNumber() + 1 );
+	}
+	SnapshotNames.Add( UniqueSnapshotName );
+
+	SnapshotsWithAllocationMap.Add( UniqueSnapshotName, AllocationMap );
+
+	TMap<FName, FCombinedAllocationInfo> SnapshotCombinedAllocations;
+	uint64 TotalAllocatedMemory = 0;
+	uint64 NumAllocations = 0;
+	GenerateScopedAllocations( AllocationMap, SnapshotCombinedAllocations, TotalAllocatedMemory, NumAllocations );
+	SnapshotsWithScopedAllocations.Add( UniqueSnapshotName, SnapshotCombinedAllocations );
+
+	// Decode callstacks.
+	// Replace encoded callstacks with human readable name. For easier debugging.
+	TMap<FString, FCombinedAllocationInfo> SnapshotDecodedCombinedAllocations;
+	for (auto& It : SnapshotCombinedAllocations)
+	{
+		const FString HumanReadableCallstack = FStatsCallstack::GetHumanReadable( It.Key );
+		SnapshotDecodedCombinedAllocations.Add( HumanReadableCallstack, It.Value );
+	}
+	SnapshotsWithDecodedScopedAllocations.Add( UniqueSnapshotName, SnapshotDecodedCombinedAllocations );
+
+	UE_LOG( LogStats, Warning, TEXT( "PrepareSnapshot: %s Alloc: %i Scoped: %i Total: %.2f MB" ), *UniqueSnapshotName.ToString(), AllocationMap.Num(), SnapshotCombinedAllocations.Num(), TotalAllocatedMemory / 1024.0f / 1024.0f );
+}
+
+void FRawStatsMemoryProfiler::CompareSnapshots_FName( const FName BeginSnaphotName, const FName EndSnaphotName, TMap<FName, FCombinedAllocationInfo>& out_Result )
+{
+	const auto BeginSnaphotPtr = SnapshotsWithScopedAllocations.Find( BeginSnaphotName );
+	const auto EndSnapshotPtr = SnapshotsWithScopedAllocations.Find( EndSnaphotName );
+	if (BeginSnaphotPtr && EndSnapshotPtr)
+	{
+		// Process data.
+		TMap<FName, FCombinedAllocationInfo> BeginSnaphot = *BeginSnaphotPtr;
+		TMap<FName, FCombinedAllocationInfo> EndSnaphot = *EndSnapshotPtr;
+		TMap<FName, FCombinedAllocationInfo> Result;
+
+		for (const auto& It : EndSnaphot)
+		{
+			const FName Callstack = It.Key;
+			const FCombinedAllocationInfo EndCombinedAlloc = It.Value;
+
+			const FCombinedAllocationInfo* BeginCombinedAllocPtr = BeginSnaphot.Find( Callstack );
+			if (BeginCombinedAllocPtr)
+			{
+				FCombinedAllocationInfo CombinedAllocation;
+				CombinedAllocation += EndCombinedAlloc;
+				CombinedAllocation -= *BeginCombinedAllocPtr;
+
+				if (CombinedAllocation.IsAlive())
+				{
+					out_Result.Add( Callstack, CombinedAllocation );
+				}
+			}
+			else
+			{
+				out_Result.Add( Callstack, EndCombinedAlloc );
+			}
+		}
+
+		// Sort by size.
+		out_Result.ValueSort( FCombinedAllocationInfoSizeGreater() );
+	}
+}
+
+void FRawStatsMemoryProfiler::CompareSnapshots_FString( const FName BeginSnaphotName, const FName EndSnaphotName, TMap<FString, FCombinedAllocationInfo>& out_Result )
+{
+	const auto BeginSnaphotPtr = SnapshotsWithDecodedScopedAllocations.Find( BeginSnaphotName );
+	const auto EndSnapshotPtr = SnapshotsWithDecodedScopedAllocations.Find( EndSnaphotName );
+	if (BeginSnaphotPtr && EndSnapshotPtr)
+	{
+		// Process data.
+		TMap<FString, FCombinedAllocationInfo> BeginSnaphot = *BeginSnaphotPtr;
+		TMap<FString, FCombinedAllocationInfo> EndSnaphot = *EndSnapshotPtr;
+
+		for (const auto& It : EndSnaphot)
+		{
+			const FString& Callstack = It.Key;
+			const FCombinedAllocationInfo EndCombinedAlloc = It.Value;
+
+			const FCombinedAllocationInfo* BeginCombinedAllocPtr = BeginSnaphot.Find( Callstack );
+			if (BeginCombinedAllocPtr)
+			{
+				FCombinedAllocationInfo CombinedAllocation;
+				CombinedAllocation += EndCombinedAlloc;
+				CombinedAllocation -= *BeginCombinedAllocPtr;
+
+				if (CombinedAllocation.IsAlive())
+				{
+					out_Result.Add( Callstack, CombinedAllocation );
+				}
+			}
+			else
+			{
+				out_Result.Add( Callstack, EndCombinedAlloc );
+			}
+		}
+
+		// Sort by size.
+		out_Result.ValueSort( FCombinedAllocationInfoSizeGreater() );
+	}
+}
\ No newline at end of file
diff --git a/Engine/Source/Developer/Profiler/Private/ProfilerRawStats.h b/Engine/Source/Developer/Profiler/Private/ProfilerRawStats.h
--- a/Engine/Source/Developer/Profiler/Private/ProfilerRawStats.h
+++ b/Engine/Source/Developer/Profiler/Private/ProfilerRawStats.h
@@ -211,8 +211,131 @@ protected:
 -----------------------------------------------------------------------------*/
 
 /** Class for managing the raw stats in the context of memory profiling. */
-class FRawStatsMemoryProfiler
+class FRawStatsMemoryProfiler : public FStatsReadFile
 {
+	friend struct FCreateStatsReader<FRawStatsMemoryProfiler>;
+	typedef FStatsReadFile Super;
 
+protected:
+
+	/** Initialization constructor. */
+	FRawStatsMemoryProfiler( const TCHAR* InFilename );
+
+	/** Called before started processing combined history. */
+	virtual void PreProcessStats() override;
+
+	/** Called after finished processing combined history. */
+	virtual void PostProcessStats() override;
+
+// 	/** Processes special message for advancing the stats frame from the game thread. */
+// 	virtual void ProcessAdvanceFrameEventGameThreadOperation( const FStatMessage& Message, const FStackState& Stack ) override
+// 	{}
+// 
+// 	/** Processes special message for advancing the stats frame from the render thread. */
+// 	virtual void ProcessAdvanceFrameEventRenderThreadOperation( const FStatMessage& Message, const FStackState& Stack ) override
+// 	{}
+// 
+// 	/** ProcessesIndicates begin of the cycle scope. */
+// 	virtual void ProcessCycleScopeStartOperation( const FStatMessage& Message, const FStackState& Stack ) override
+// 	{}
+// 
+// 	/** Indicates end of the cycle scope. */
+// 	virtual void ProcessCycleScopeEndOperation( const FStatMessage& Message, const FStackState& Stack ) override
+// 	{}
+// 
+	/** Processes special message marker used determine that we encountered a special data in the stat file. */
+	virtual void ProcessSpecialMessageMarkerOperation( const FStatMessage& Message, const FStackState& StackState ) override;
+// 
+// 	/** Processes set operation. */
+// 	virtual void ProcessSetOperation( const FStatMessage& Message, const FStackState& Stack ) override
+// 	{}
+// 
+// 	/** Processes clear operation. */
+// 	virtual void ProcessClearOperation( const FStatMessage& Message, const FStackState& Stack ) override
+// 	{}
+// 
+// 	/** Processes add operation. */
+// 	virtual void ProcessAddOperation( const FStatMessage& Message, const FStackState& Stack ) override
+// 	{}
+// 
+// 	/** Processes subtract operation. */
+// 	virtual void ProcessSubtractOperation( const FStatMessage& Message, const FStackState& Stack ) override
+// 	{}
+
+	/** Processes memory operation. @see EMemoryOperation. */
+	virtual void ProcessMemoryOperation( EMemoryOperation MemOp, uint64 Ptr, uint64 NewPtr, int64 Size, uint32 SequenceTag, const FStackState& StackState ) override;
+
+	/**
+	 * Sorts sequence allocations.
+	 *
+	 * @return true, if finished, false otherwise like stopped.
+	 */
+	bool SortSequenceAllocations();
+
+public:
+	/** Generates a basic memory usage report and prints it to the log. */
+	void GenerateMemoryUsageReport( const TMap<uint64, FAllocationInfo>& AllocationMap, const TSet<FName>& UObjectRawNames );
+
+	/** Generates scoped allocation statistics. */
+	void ProcessAndDumpScopedAllocations( const TMap<uint64, FAllocationInfo>& AllocationMap );
+
+	/** Generates UObject allocation statistics. */
+	void ProcessAndDumpUObjectAllocations( const TMap<uint64, FAllocationInfo>& AllocationMap, const TSet<FName>& UObjectRawNames );
+
+	void DumpScopedAllocations( const TCHAR* Name, const TMap<FString, FCombinedAllocationInfo>& CombinedAllocations );
+
+	/** Generates callstack based allocation map. */
+	void GenerateScopedAllocations( const TMap<uint64, FAllocationInfo>& AllocationMap, TMap<FName, FCombinedAllocationInfo>& out_CombinedAllocations, uint64& TotalAllocatedMemory, uint64& NumAllocations );
+
+	void GenerateScopedTreeAllocations( const TMap<FName, FCombinedAllocationInfo>& ScopedAllocations, FNodeAllocationInfo& out_Root );
+
+	/** Prepares data for a snapshot. */
+	void PrepareSnapshot( const FName SnapshotName, const TMap<uint64, FAllocationInfo>& AllocationMap );
+
+	void CompareSnapshots_FName( const FName BeginSnaphotName, const FName EndSnaphotName, TMap<FName, FCombinedAllocationInfo>& out_Result );
+
+	void CompareSnapshots_FString( const FName BeginSnaphotName, const FName EndSnaphotName, TMap<FString, FCombinedAllocationInfo>& out_Result );
+
+
+	/**
+	* @return Platform's name based on the loaded ue4statsraw file
+	*/
+	const FString& GetPlatformName() const
+	{
+		return Header.PlatformName;
+	}
+
+protected:
+	/**
+	 *	 All allocation ordered by the sequence tag.
+	 *	 There is an assumption that the sequence tag will not turn-around.
+	 */
+	TArray<FAllocationInfo> SequenceAllocationArray;
+
+	/** Allocation map. */
+	//TMap<uint64, FAllocationInfo>& AllocationMap;
+
+	/** The sequence tag mapping to the named markers. */
+	TArray<TPair<uint32, FName>> Snapshots;
+
+	/** Unique snapshot names. */
+	TSet<FName> SnapshotNames;
+
+	/** The sequence tag mapping to the named markers that need to processed. */
+	TArray<TPair<uint32, FName>> SnapshotsToBeProcessed;
+
+	/** Snapshots with allocation map. */
+	TMap<FName, TMap<uint64, FAllocationInfo> > SnapshotsWithAllocationMap;
+
+	/** Snapshots with callstack based allocation map. */
+	TMap<FName, TMap<FName, FCombinedAllocationInfo> > SnapshotsWithScopedAllocations;
+
+	TMap<FName, TMap<FString, FCombinedAllocationInfo> > SnapshotsWithDecodedScopedAllocations;
+
+	/** Number of memory operations. */
+	uint64 NumMemoryOperations;
+
+	/** Last sequence tag for named marker. */
+	uint32 LastSequenceTagForNamedMarker;
 };
 
diff --git a/Engine/Source/Developer/Profiler/Public/Interfaces/IProfilerModule.h b/Engine/Source/Developer/Profiler/Public/Interfaces/IProfilerModule.h
--- a/Engine/Source/Developer/Profiler/Public/Interfaces/IProfilerModule.h
+++ b/Engine/Source/Developer/Profiler/Public/Interfaces/IProfilerModule.h
@@ -25,5 +25,5 @@ public:
 	virtual void StatsMemoryDumpCommand( const TCHAR* Filename ) = 0;
 
 	/** Creates a new instance of the memory profiler based the raw stats. */
-	virtual FRawStatsMemoryProfiler* OpenRawStatForMemoryProfiling( const TCHAR* Filename ) = 0;
+	virtual FRawStatsMemoryProfiler* OpenRawStatsForMemoryProfiling( const TCHAR* Filename ) = 0;
 };
\ No newline at end of file
diff --git a/Engine/Source/Programs/UnrealFrontend/Private/Commands/StatsDumpMemoryCommand.cpp b/Engine/Source/Programs/UnrealFrontend/Private/Commands/StatsDumpMemoryCommand.cpp
--- a/Engine/Source/Programs/UnrealFrontend/Private/Commands/StatsDumpMemoryCommand.cpp
+++ b/Engine/Source/Programs/UnrealFrontend/Private/Commands/StatsDumpMemoryCommand.cpp
@@ -1,1225 +1,20 @@
 // Copyright 1998-2015 Epic Games, Inc. All Rights Reserved.
 
 #include "UnrealFrontendPrivatePCH.h"
-#include "Profiler.h"
 #include "StatsDumpMemoryCommand.h"
-#include "DiagnosticTable.h"
-
-/*-----------------------------------------------------------------------------
-	Sort helpers
------------------------------------------------------------------------------*/
-
-/** Sorts allocations by size. */
-struct FAllocationInfoSequenceTagLess
-{
-	FORCEINLINE bool operator()( const FAllocationInfo& A, const FAllocationInfo& B ) const
-	{
-		return A.SequenceTag < B.SequenceTag;
-	}
-};
-
-/** Sorts allocations by size. */
-struct FAllocationInfoSizeGreater
-{
-	FORCEINLINE bool operator()( const FAllocationInfo& A, const FAllocationInfo& B ) const
-	{
-		return B.Size < A.Size;
-	}
-};
-
-/** Sorts combined allocations by size. */
-struct FCombinedAllocationInfoSizeGreater
-{
-	FORCEINLINE bool operator()( const FCombinedAllocationInfo& A, const FCombinedAllocationInfo& B ) const
-	{
-		return B.Size < A.Size;
-	}
-};
-
-
-/** Sorts node allocations by size. */
-struct FNodeAllocationInfoSizeGreater
-{
-	FORCEINLINE bool operator()( const FNodeAllocationInfo& A, const FNodeAllocationInfo& B ) const
-	{
-		return B.Size < A.Size;
-	}
-};
-
-/*-----------------------------------------------------------------------------
-	Callstack decoding/encoding
------------------------------------------------------------------------------*/
-
-/** Helper struct used to manipulate stats based callstacks. */
-struct FStatsCallstack
-{
-	/** Separator. */
-	static const TCHAR* CallstackSeparator;
-
-	/** Encodes decoded callstack a string, to be like '45+656+6565'. */
-	static FString Encode( const TArray<FName>& Callstack )
-	{
-		FString Result;
-		for (const auto& Name : Callstack)
-		{
-			Result += TTypeToString<int32>::ToString( (int32)Name.GetComparisonIndex() );
-			Result += CallstackSeparator;
-		}
-		return Result;
-	}
-
-	/** Decodes encoded callstack to an array of FNames. */
-	static void DecodeToNames( const FName& EncodedCallstack, TArray<FName>& out_DecodedCallstack )
-	{
-		TArray<FString> DecodedCallstack;
-		DecodeToStrings( EncodedCallstack, DecodedCallstack );
-
-		// Convert back to FNames
-		for (const auto& It : DecodedCallstack)
-		{
-			NAME_INDEX NameIndex = 0;
-			TTypeFromString<NAME_INDEX>::FromString( NameIndex, *It );
-			const FName LongName = FName( NameIndex, NameIndex, 0 );
-
-			out_DecodedCallstack.Add( LongName );
-		}
-	}
-
-	/** Converts the encoded callstack into human readable callstack. */
-	static FString GetHumanReadable( const FName& EncodedCallstack )
-	{
-		TArray<FName> DecodedCallstack;
-		DecodeToNames( EncodedCallstack, DecodedCallstack );
-		const FString Result = GetHumanReadable( DecodedCallstack );
-		return Result;
-	}
-
-	/** Converts the encoded callstack into human readable callstack. */
-	static FString GetHumanReadable( const TArray<FName>& DecodedCallstack )
-	{
-		FString Result;
-
-		const int32 NumEntries = DecodedCallstack.Num();
-		//for (int32 Index = DecodedCallstack.Num() - 1; Index >= 0; --Index)
-		for (int32 Index = 0; Index < NumEntries; ++Index)
-		{
-			const FName LongName = DecodedCallstack[Index];
-			const FString ShortName = FStatNameAndInfo::GetShortNameFrom( LongName ).ToString();
-			//const FString Group = FStatNameAndInfo::GetGroupNameFrom( LongName ).ToString();
-			FString Desc = FStatNameAndInfo::GetDescriptionFrom( LongName );
-			Desc.Trim();
-
-			if (Desc.Len() == 0)
-			{
-				Result += ShortName;
-			}
-			else
-			{
-				Result += Desc;
-			}
-
-			if (Index != NumEntries - 1)
-			{
-				Result += TEXT( " -> " );
-			}
-		}
-
-		Result.ReplaceInline( TEXT( "STAT_" ), TEXT( "" ), ESearchCase::CaseSensitive );
-		return Result;
-	}
-
-protected:
-	/** Decodes encoded callstack to an array of strings. Where each string is the index of the FName. */
-	static void DecodeToStrings( const FName& EncodedCallstack, TArray<FString>& out_DecodedCallstack )
-	{
-		EncodedCallstack.ToString().ParseIntoArray( out_DecodedCallstack, CallstackSeparator, true );
-	}
-};
-
-const TCHAR* FStatsCallstack::CallstackSeparator = TEXT( "+" );
-
-/*-----------------------------------------------------------------------------
-	Allocation info
------------------------------------------------------------------------------*/
-
-FAllocationInfo::FAllocationInfo( uint64 InOldPtr, uint64 InPtr, int64 InSize, const TArray<FName>& InCallstack, uint32 InSequenceTag, EMemoryOperation InOp, bool bInHasBrokenCallstack ) 
-	: OldPtr( InOldPtr )
-	, Ptr( InPtr )
-	, Size( InSize )
-	, EncodedCallstack( *FStatsCallstack::Encode( InCallstack ) )
-	, SequenceTag( InSequenceTag )
-	, Op( InOp )
-	, bHasBrokenCallstack( bInHasBrokenCallstack )
-{
-
-}
-
-FAllocationInfo::FAllocationInfo( const FAllocationInfo& Other ) 
-	: OldPtr( Other.OldPtr )
-	, Ptr( Other.Ptr )
-	, Size( Other.Size )
-	, EncodedCallstack( Other.EncodedCallstack )
-	, SequenceTag( Other.SequenceTag )
-	, Op( Other.Op )
-	, bHasBrokenCallstack( Other.bHasBrokenCallstack )
-{
-
-}
-
-/*-----------------------------------------------------------------------------
-	FNodeAllocationInfo
------------------------------------------------------------------------------*/
-
-void FNodeAllocationInfo::SortBySize()
-{
-	ChildNodes.ValueSort( FNodeAllocationInfoSizeGreater() );
-	for (auto& It : ChildNodes)
-	{
-		It.Value->SortBySize();
-	}
-}
-
-
-void FNodeAllocationInfo::PrepareCallstackData( const TArray<FName>& InDecodedCallstack )
-{
-	DecodedCallstack = InDecodedCallstack;
-	EncodedCallstack = *FStatsCallstack::Encode( DecodedCallstack );
-	HumanReadableCallstack = FStatsCallstack::GetHumanReadable( DecodedCallstack );
-}
-
-/*-----------------------------------------------------------------------------
-	Stats stack helpers
------------------------------------------------------------------------------*/
-
-/** Holds stats stack state, used to preserve continuity when the game frame has changed. */
-struct FStackState
-{
-	/** Default constructor. */
-	FStackState()
-		: bIsBrokenCallstack( false )
-	{}
-
-	/** Call stack. */
-	TArray<FName> Stack;
-
-	/** Current function name. */
-	FName Current;
-
-	/** Whether this callstack is marked as broken due to mismatched start and end scope cycles. */
-	bool bIsBrokenCallstack;
-};
-
-
+#include "Profiler.h"
 
 /*-----------------------------------------------------------------------------
 	FStatsMemoryDumpCommand
 -----------------------------------------------------------------------------*/
 
-static const double NumSecondsBetweenLogs = 5.0;
-
-void FStatsMemoryDumpCommand::InternalRun()
+void FStatsMemoryDumpCommand::Run()
 {
 	FString SourceFilepath;
 	FParse::Value( FCommandLine::Get(), TEXT( "-INFILE=" ), SourceFilepath );
 
-	{
-		const FName NAME_ProfilerModule = TEXT( "Profiler" );
-		IProfilerModule& ProfilerModule = FModuleManager::LoadModuleChecked<IProfilerModule>( NAME_ProfilerModule );
-		ProfilerModule.StatsMemoryDumpCommand( *SourceFilepath );
-		//FModuleManager::Get().UnloadModule( NAME_ProfilerModule );
-	}
-
-	{
-		const FName NAME_ProfilerModule = TEXT( "Profiler" );
-		IProfilerModule& ProfilerModule = FModuleManager::LoadModuleChecked<IProfilerModule>( NAME_ProfilerModule );
-		ProfilerModule.StatsMemoryDumpCommand( *SourceFilepath );
-		//FModuleManager::Get().UnloadModule( NAME_ProfilerModule );
-	}
-
-
-	FStatsReadFile* NewReader = FStatsReadFile::CreateReaderForRawStats( *SourceFilepath, FStatsReadFile::FOnNewCombinedHistory::CreateRaw( this, &FStatsMemoryDumpCommand::ProcessMemoryOperations ) );
-	if (NewReader)
-	{
-		PlatformName = NewReader->Header.PlatformName;
-		NewReader->ReadAndProcessAsynchronously();
-
-		while( NewReader->IsBusy() )
-		{
-			FPlatformProcess::Sleep( 1.0f );
-
-			UE_LOG( LogStats, Log, TEXT( "Async: Stage: %s / %3i%%" ), *NewReader->GetProcessingStageAsString(), NewReader->GetStageProgress() );
-		}
-		//NewReader->RequestStop();
-
-		// Delete reader, we don't need the data anymore.
-		delete NewReader;
-
-		// Frame-240 Frame-120 Frame-060
-		TMap<FString, FCombinedAllocationInfo> FrameBegin_Exit;
-		CompareSnapshots_FString( TEXT( "BeginSnapshot" ), TEXT( "EngineLoop.Exit" ), FrameBegin_Exit );
-		DumpScopedAllocations( TEXT( "Begin_Exit" ), FrameBegin_Exit );
-
-#if	UE_BUILD_DEBUG
-		TMap<FString, FCombinedAllocationInfo> Frame060_120;
-		CompareSnapshots_FString( TEXT( "Frame-060" ), TEXT( "Frame-120" ), Frame060_120 );
-		DumpScopedAllocations( TEXT( "Frame060_120" ), Frame060_120 );
-
-		TMap<FString, FCombinedAllocationInfo> Frame060_240;
-		CompareSnapshots_FString( TEXT( "Frame-060" ), TEXT( "Frame-240" ), Frame060_240 );
-		DumpScopedAllocations( TEXT( "Frame060_240" ), Frame060_240 );
-
-		// Generate scoped tree view.
-		{
-			TMap<FName, FCombinedAllocationInfo> FrameBegin_Exit_FName;
-			CompareSnapshots_FName( TEXT( "BeginSnapshot" ), TEXT( "EngineLoop.Exit" ), FrameBegin_Exit_FName );
-
-			FNodeAllocationInfo Root;
-			Root.EncodedCallstack = TEXT( "ThreadRoot" );
-			Root.HumanReadableCallstack = TEXT( "ThreadRoot" );
-			GenerateScopedTreeAllocations( FrameBegin_Exit_FName, Root );
-		}
-
-
-		{
-			TMap<FName, FCombinedAllocationInfo> Frame060_240_FName;
-			CompareSnapshots_FName( TEXT( "Frame-060" ), TEXT( "Frame-240" ), Frame060_240_FName );
-
-			FNodeAllocationInfo Root;
-			Root.EncodedCallstack = TEXT( "ThreadRoot" );
-			Root.HumanReadableCallstack = TEXT( "ThreadRoot" );
-			GenerateScopedTreeAllocations( Frame060_240_FName, Root );
-		}
-#endif // UE_BUILD_DEBUG
-	}
-}
-
-
-void FStatsMemoryDumpCommand::ProcessMemoryOperations( FStatsReadFile* InFile )
-{
-	const FStatsReadFile& File = *InFile;
-	const TMap<int64, FStatPacketArray>& CombinedHistory = File.CombinedHistory;
-
-	// This is only example code, no fully implemented, may sometimes crash.
-	// This code is not optimized. 
-	double PreviousSeconds = FPlatformTime::Seconds();
-	int32 CurrentStatMessageIndex = 0;
-	uint64 NumMemoryOperations = 0;
-
-	// Generate frames
-	TArray<int64> Frames;
-	CombinedHistory.GenerateKeyArray( Frames );
-	Frames.Sort();
-
-	// Raw stats callstack for this stat packet array.
-	TMap<FName, FStackState> StackStates;
-
-	// All allocation ordered by the sequence tag.
-	// There is an assumption that the sequence tag will not turn-around.
-	//TMap<uint32, FAllocationInfo> SequenceAllocationMap;
-	TArray<FAllocationInfo> SequenceAllocationArray;
-
-	// Pass 1.
-	// Read all stats messages, parse all memory operations and decode callstacks.
-	const int64 FirstFrame = 0;
-	PreviousSeconds -= NumSecondsBetweenLogs;
-
-	// Begin marker.
-	uint32 LastSequenceTagForNamedMarker = 0;
-	Snapshots.Add( TPairInitializer<uint32, FName>( LastSequenceTagForNamedMarker, TEXT( "BeginSnapshot" ) ) );
-
-	const int32 OnerPercent = FMath::Max( int32(InFile->FileInfo.TotalStatMessagesNum / 100), 65536 );
-
-	// #YRX_Stats: 2015-07-10 This still should be generic
-	for( int32 FrameIndex = 0; FrameIndex < Frames.Num(); ++FrameIndex )
-	{
-		const int64 TargetFrame = Frames[FrameIndex];
-		const int64 Diff = TargetFrame - FirstFrame;
-		const FStatPacketArray& Frame = CombinedHistory.FindChecked( TargetFrame );
-
-		bool bAtLeastOnePacket = false;
-		for( int32 PacketIndex = 0; PacketIndex < Frame.Packets.Num(); PacketIndex++ )
-		{
-			const FStatPacket& StatPacket = *Frame.Packets[PacketIndex];
-			const FName& ThreadFName = File.State.Threads.FindChecked( StatPacket.ThreadId );
-
-			FStackState* StackState = StackStates.Find( ThreadFName );
-			if( !StackState )
-			{
-				StackState = &StackStates.Add( ThreadFName );
-				StackState->Stack.Add( ThreadFName );
-				StackState->Current = ThreadFName;
-			}
-
-			const FStatMessagesArray& Data = StatPacket.StatMessages;
-
-			const int32 NumStatMessages = Data.Num();
-			for( int32 Index = 0; Index < NumStatMessages; Index++ )
-			{
-				CurrentStatMessageIndex++;
-
-				if (CurrentStatMessageIndex % OnerPercent == 0)
-				{
-					const double CurrentSeconds = FPlatformTime::Seconds();
-					if( CurrentSeconds > PreviousSeconds + NumSecondsBetweenLogs )
-					{
-						const int32 PercentagePos = int32( 100.0*CurrentStatMessageIndex / InFile->FileInfo.TotalStatMessagesNum );
-						InFile->StageProgress.Set( PercentagePos );
-						UE_LOG( LogStats, Verbose, TEXT( "Processing %3i%% (%10i/%10i) stat messages [Frame: %3i, Packet: %2i]" ), PercentagePos, CurrentStatMessageIndex, InFile->FileInfo.TotalStatMessagesNum, FrameIndex, PacketIndex );
-						PreviousSeconds = CurrentSeconds;
-
-						// Abandon support, simply return for now.
-						if (InFile->bShouldStopProcessing == true)
-						{
-							InFile->SetProcessingStage( EStatsProcessingStage::SPS_Stopped );
-							return;
-						}
-					}
-				}
-
-				const FStatMessage& Item = Data[Index];
-				const EStatOperation::Type Op = Item.NameAndInfo.GetField<EStatOperation>();
-				const FName RawName = Item.NameAndInfo.GetRawName();
-
-				if (Op == EStatOperation::CycleScopeStart || Op == EStatOperation::CycleScopeEnd || Op == EStatOperation::Memory || Op == EStatOperation::SpecialMessageMarker)
-				{
-					if( Op == EStatOperation::CycleScopeStart )
-					{
-						StackState->Stack.Add( RawName );
-						StackState->Current = RawName;
-					}
-					else if( Op == EStatOperation::Memory )
-					{
-						// Experimental code used only to test the implementation.
-						// First memory operation is Alloc or Free
-						const uint64 EncodedPtr = Item.GetValue_Ptr();
-						const EMemoryOperation MemOp = EMemoryOperation( EncodedPtr & (uint64)EMemoryOperation::Mask );
-						const uint64 Ptr = EncodedPtr & ~(uint64)EMemoryOperation::Mask;
-						if (MemOp == EMemoryOperation::Alloc)
-						{
-							NumMemoryOperations++;
-							// @see FStatsMallocProfilerProxy::TrackAlloc
-							// After AllocPtr message there is always alloc size message and the sequence tag.
-							Index++; CurrentStatMessageIndex++;
-							const FStatMessage& AllocSizeMessage = Data[Index];
-							const int64 AllocSize = AllocSizeMessage.GetValue_int64();
-
-							// Read OperationSequenceTag.
-							Index++; CurrentStatMessageIndex++;
-							const FStatMessage& SequenceTagMessage = Data[Index];
-							const uint32 SequenceTag = SequenceTagMessage.GetValue_int64();
-
-							//ThreadStats->AddMemoryMessage( GET_STATFNAME( STAT_Memory_AllocPtr ), (uint64)(UPTRINT)Ptr | (uint64)EMemoryOperation::Alloc );
-							//ThreadStats->AddMemoryMessage( GET_STATFNAME( STAT_Memory_AllocSize ), Size );
-							//ThreadStats->AddMemoryMessage( GET_STATFNAME( STAT_Memory_OperationSequenceTag ), (int64)SequenceTag );
-
-							// Add a new allocation.
-							SequenceAllocationArray.Add(
-								FAllocationInfo(
-								0,
-								Ptr,
-								AllocSize,
-								StackState->Stack,
-								SequenceTag,
-								EMemoryOperation::Alloc,
-								StackState->bIsBrokenCallstack
-								) );
-							LastSequenceTagForNamedMarker = SequenceTag;
-
-						}
-						else if (MemOp == EMemoryOperation::Realloc)
-						{
-							NumMemoryOperations++;
-							const uint64 OldPtr = Ptr;
-
-							// Read NewPtr
-							Index++; CurrentStatMessageIndex++;
-							const FStatMessage& AllocPtrMessage = Data[Index];
-							const uint64 NewPtr = AllocPtrMessage.GetValue_Ptr() & ~(uint64)EMemoryOperation::Mask;
-
-							// After AllocPtr message there is always alloc size message and the sequence tag.
-							Index++; CurrentStatMessageIndex++;
-							const FStatMessage& ReallocSizeMessage = Data[Index];
-							const int64 ReallocSize = ReallocSizeMessage.GetValue_int64();
-
-							// Read OperationSequenceTag.
-							Index++; CurrentStatMessageIndex++;
-							const FStatMessage& SequenceTagMessage = Data[Index];
-							const uint32 SequenceTag = SequenceTagMessage.GetValue_int64();
-
-							//ThreadStats->AddMemoryMessage( GET_STATFNAME( STAT_Memory_FreePtr ), (uint64)(UPTRINT)OldPtr | (uint64)EMemoryOperation::Realloc );
-							//ThreadStats->AddMemoryMessage( GET_STATFNAME( STAT_Memory_AllocPtr ), (uint64)(UPTRINT)NewPtr | (uint64)EMemoryOperation::Realloc );
-							//ThreadStats->AddMemoryMessage( GET_STATFNAME( STAT_Memory_AllocSize ), NewSize );
-							//ThreadStats->AddMemoryMessage( GET_STATFNAME( STAT_Memory_OperationSequenceTag ), (int64)SequenceTag );
-
-							// Add a new realloc.
-							SequenceAllocationArray.Add(
-								FAllocationInfo(
-								OldPtr, 
-								NewPtr,
-								ReallocSize,
-								StackState->Stack,
-								SequenceTag,
-								EMemoryOperation::Realloc,
-								StackState->bIsBrokenCallstack
-								) );
-							LastSequenceTagForNamedMarker = SequenceTag;
-						}
-						else if (MemOp == EMemoryOperation::Free)
-						{
-							NumMemoryOperations++;
-							// Read OperationSequenceTag.
-							Index++; CurrentStatMessageIndex++;
-							const FStatMessage& SequenceTagMessage = Data[Index];
-							const uint32 SequenceTag = SequenceTagMessage.GetValue_int64();
-
-							//ThreadStats->AddMemoryMessage( GET_STATFNAME( STAT_Memory_FreePtr ), (uint64)(UPTRINT)Ptr | (uint64)EMemoryOperation::Free );	// 16 bytes total				
-							//ThreadStats->AddMemoryMessage( GET_STATFNAME( STAT_Memory_OperationSequenceTag ), (int64)SequenceTag );
-
-							// Add a new free.
-							SequenceAllocationArray.Add(
-								FAllocationInfo(
-								0,
-								Ptr,		
-								0,
-								StackState->Stack,
-								SequenceTag,
-								EMemoryOperation::Free,
-								StackState->bIsBrokenCallstack
-								) );
-						}
-						else
-						{
-							UE_LOG( LogStats, Warning, TEXT( "Pointer from a memory operation is invalid" ) );
-						}
-					}
-					else if( Op == EStatOperation::CycleScopeEnd )
-					{
-						if( StackState->Stack.Num() > 1 )
-						{
-							const FName ScopeStart = StackState->Stack.Pop();
-							const FName ScopeEnd = Item.NameAndInfo.GetRawName();
-
-							check( ScopeStart == ScopeEnd );
-
-							StackState->Current = StackState->Stack.Last();
-
-							// The stack should be ok, but it may be partially broken.
-							// This will happen if memory profiling starts in the middle of executing a background thread.
-							StackState->bIsBrokenCallstack = false;
-						}
-						else
-						{
-							const FName ShortName = Item.NameAndInfo.GetShortName();
-
-							UE_LOG( LogStats, Warning, TEXT( "Broken cycle scope end %s/%s, current %s" ),
-									*ThreadFName.ToString(),
-									*ShortName.ToString(),
-									*StackState->Current.ToString() );
-
-							// The stack is completely broken, only has the thread name and the last cycle scope.
-							// Rollback to the thread node.
-							StackState->bIsBrokenCallstack = true;
-							StackState->Stack.Empty();
-							StackState->Stack.Add( ThreadFName );
-							StackState->Current = ThreadFName;
-						}
-					}
-					else if (Op == EStatOperation::SpecialMessageMarker)
-					{
-						if (RawName == FStatConstants::NAME_NamedMarker)
-						{
-							const FName NamedMarker = Item.GetValue_FName(); // LastSequenceTag
-							Snapshots.Add( TPairInitializer<uint32, FName>( LastSequenceTagForNamedMarker, NamedMarker ) );
-						}
-					}
-				}
-			}
-		}
-	}
-
-	InFile->StageProgress.Set( 100 );
-
-	// End marker.
-	Snapshots.Add( TPairInitializer<uint32, FName>( TNumericLimits<uint32>::Max(), TEXT( "EndSnapshot" ) ) );
-
-	// Copy snapshots.
-	SnapshotsToBeProcessed = Snapshots;
-
-	UE_LOG( LogStats, Warning, TEXT( "NumMemoryOperations:   %llu" ), NumMemoryOperations );
-	UE_LOG( LogStats, Warning, TEXT( "SequenceAllocationNum: %i" ), SequenceAllocationArray.Num() );
-
-	// Pass 2.
-	/*
-	TMap<uint32,FAllocationInfo> UniqueSeq;
-	TMultiMap<uint32,FAllocationInfo> OriginalAllocs;
-	TMultiMap<uint32,FAllocationInfo> BrokenAllocs;
-	for( const FAllocationInfo& Alloc : SequenceAllocationArray )
-	{
-		const FAllocationInfo* Found = UniqueSeq.Find(Alloc.SequenceTag);
-		if( !Found )
-		{
-			UniqueSeq.Add(Alloc.SequenceTag,Alloc);
-		}
-		else
-		{
-			OriginalAllocs.Add(Alloc.SequenceTag, *Found);
-			BrokenAllocs.Add(Alloc.SequenceTag, Alloc);
-		}
-	}
-	*/
-
-	InFile->SetProcessingStage( EStatsProcessingStage::SPS_SortSequence );
-
-	// Sort all memory operation by the sequence tag, iterate through all operation and generate memory usage.
-	SequenceAllocationArray.Sort( FAllocationInfoSequenceTagLess() );
-
-	// Abandon support, simply return for now.
-	if (InFile->bShouldStopProcessing == true)
-	{
-		InFile->SetProcessingStage( EStatsProcessingStage::SPS_Stopped );
-		return;
-	}
-
-	// Named markers/snapshots
-
-	// Alive allocations.
-	TMap<uint64, FAllocationInfo> AllocationMap;
-	TMultiMap<uint64, FAllocationInfo> FreeWithoutAllocMap;
-	TMultiMap<uint64, FAllocationInfo> DuplicatedAllocMap;
-	int32 NumDuplicatedMemoryOperations = 0;
-	int32 NumFWAMemoryOperations = 0; // FreeWithoutAlloc
-	int32 NumZeroAllocs = 0; // Malloc(0)
-
-	InFile->SetProcessingStage( EStatsProcessingStage::SPS_ProcessAllocations );
-
-	// Initialize the begin snapshot.
-	auto BeginSnapshot = SnapshotsToBeProcessed[0];
-	SnapshotsToBeProcessed.RemoveAt( 0 );
-	PrepareSnapshot( BeginSnapshot.Value, AllocationMap );
-
-	
-	auto CurrentSnapshot = SnapshotsToBeProcessed[0];
-
-	UE_LOG( LogStats, Warning, TEXT( "Generating memory operations map" ) );
-	const int32 NumSequenceAllocations = SequenceAllocationArray.Num();
-	const int32 OnePercent = FMath::Max( NumSequenceAllocations / 100, 1024 );
-	for( int32 Index = 0; Index < NumSequenceAllocations; Index++ )
-	{
-		if( Index % OnePercent )
-		{
-			const double CurrentSeconds = FPlatformTime::Seconds();
-			if( CurrentSeconds > PreviousSeconds + NumSecondsBetweenLogs )
-			{
-				const int32 PercentagePos = int32( 100.0*(Index + 1) / NumSequenceAllocations );
-				InFile->StageProgress.Set( PercentagePos );
-				UE_LOG( LogStats, Verbose, TEXT( "Processing allocations %3i%% (%10i/%10i)" ), PercentagePos, Index + 1, NumSequenceAllocations );
-				PreviousSeconds = CurrentSeconds;
-
-				// Abandon support, simply return for now.
-				if (InFile->bShouldStopProcessing == true)
-				{
-					InFile->SetProcessingStage( EStatsProcessingStage::SPS_Stopped );
-					return;
-				}
-			}
-		}
-
-		const FAllocationInfo& Alloc = SequenceAllocationArray[Index];
-
-		// Check named marker/snapshots
-		if (Alloc.SequenceTag > CurrentSnapshot.Key)
-		{
-			SnapshotsToBeProcessed.RemoveAt( 0 );
-			PrepareSnapshot( CurrentSnapshot.Value, AllocationMap );
-			CurrentSnapshot = SnapshotsToBeProcessed[0];
-		}
-
-		if (Alloc.Op == EMemoryOperation::Alloc)
-		{
-			if (Alloc.Size == 0)
-			{
-				NumZeroAllocs++;
-			}
-
-			const FAllocationInfo* Found = AllocationMap.Find( Alloc.Ptr );
-
-			if( !Found )
-			{
-				AllocationMap.Add( Alloc.Ptr, Alloc );
-			}
-			else
-			{
-				const FAllocationInfo* FoundAndFreed = FreeWithoutAllocMap.Find( Found->Ptr );
-				const FAllocationInfo* FoundAndAllocated = FreeWithoutAllocMap.Find( Alloc.Ptr );
-
-#if	UE_BUILD_DEBUG
-				if( FoundAndFreed )
-				{
-					const FString FoundAndFreedCallstack = FStatsCallstack::GetHumanReadable( FoundAndFreed->EncodedCallstack );
-				}
-
-				if( FoundAndAllocated )
-				{
-					const FString FoundAndAllocatedCallstack = FStatsCallstack::GetHumanReadable( FoundAndAllocated->EncodedCallstack );
-				}
-
-				NumDuplicatedMemoryOperations++;
-
-
-				const FString FoundCallstack = FStatsCallstack::GetHumanReadable( Found->EncodedCallstack );
-				const FString AllocCallstack = FStatsCallstack::GetHumanReadable( Alloc.EncodedCallstack );
-				UE_LOG( LogStats, Warning, TEXT( "DuplicateAlloc: %s Size: %i/%i Ptr: %i/%i Tag: %i/%i" ), *AllocCallstack, Found->Size, Alloc.Size, Found->Ptr, Alloc.Ptr, Found->SequenceTag, Alloc.SequenceTag );
-#endif // UE_BUILD_DEBUG
-
-				// Replace pointer.
-				AllocationMap.Add( Alloc.Ptr, Alloc );
-				// Store the old pointer.
-				DuplicatedAllocMap.Add( Alloc.Ptr, *Found );
-			}
-		}
-		else if (Alloc.Op == EMemoryOperation::Realloc)
-		{
-			if (Alloc.Size == 0)
-			{
-				NumZeroAllocs++;
-			}
-
-			// Previous Alloc or Realloc
-			const FAllocationInfo* FoundOld = AllocationMap.Find( Alloc.OldPtr );
-
-			if (FoundOld)
-			{
-				const bool bIsValid = Alloc.SequenceTag > FoundOld->SequenceTag;
-				if (!bIsValid)
-				{
-					UE_LOG( LogStats, Warning, TEXT( "InvalidRealloc Ptr: %llu, Seq: %i/%i" ), Alloc.Ptr, Alloc.SequenceTag, FoundOld->SequenceTag );
-				}
-				// Remove the old allocation.
-				AllocationMap.Remove( Alloc.OldPtr );
-				AllocationMap.Add( Alloc.Ptr, Alloc );
-			}
-			// If we have situation where the old pointer is the same as the new pointer
-			// There is high change that the original call looked like this 
-			// Ptr = Malloc( 40 ); Ptr = Realloc( Ptr, 40 );
-			else if (Alloc.OldPtr != Alloc.Ptr)
-			{
-				// No OldPtr, should not happen as it means realloc without initial alloc.
-				// Or Realloc after Malloc(0)
-
-#if	UE_BUILD_DEBUG
-				const FString ReallocCallstack = FStatsCallstack::GetHumanReadable( Alloc.EncodedCallstack );
-				//UE_LOG( LogStats, Warning, TEXT( "ReallocWithoutAlloc: %s %i %i/%i [%i]" ), *ReallocCallstack, Alloc.Size, Alloc.OldPtr, Alloc.Ptr, Alloc.SequenceTag );
-#endif // UE_BUILD_DEBUG
-
-				AllocationMap.Add( Alloc.Ptr, Alloc );
-			}
-
-		}
-		else if (Alloc.Op == EMemoryOperation::Free)
-		{
-			const FAllocationInfo* Found = AllocationMap.Find( Alloc.Ptr );
-			if( Found )
-			{
-				const bool bIsValid = Alloc.SequenceTag > Found->SequenceTag;
-				if( !bIsValid )
-				{
-					UE_LOG( LogStats, Warning, TEXT( "InvalidFree Ptr: %llu, Seq: %i/%i" ), Alloc.Ptr, Alloc.SequenceTag, Found->SequenceTag );
-				}
-				AllocationMap.Remove( Alloc.Ptr );
-			}
-			else
-			{
-				FreeWithoutAllocMap.Add( Alloc.Ptr, Alloc );
-				NumFWAMemoryOperations++;
-
-#if	UE_BUILD_DEBUG
-				const FString FWACallstack = FStatsCallstack::GetHumanReadable( Alloc.EncodedCallstack );
-				//UE_LOG( LogStats, Warning, TEXT( "FreeWithoutAllocCallstack: %s %i" ), *FWACallstack, Alloc.Ptr );
-#endif // UE_BUILD_DEBUG
-			}
-		}
-	}
-	InFile->StageProgress.Set( 100 );
-
-	auto EndSnapshot = SnapshotsToBeProcessed[0];
-	SnapshotsToBeProcessed.RemoveAt( 0 );
-	PrepareSnapshot( EndSnapshot.Value, AllocationMap );
-
-	UE_LOG( LogStats, Warning, TEXT( "NumDuplicatedMemoryOperations: %i" ), NumDuplicatedMemoryOperations );
-	UE_LOG( LogStats, Warning, TEXT( "NumFWAMemoryOperations:        %i" ), NumFWAMemoryOperations );
-	UE_LOG( LogStats, Warning, TEXT( "NumZeroAllocs:                 %i" ), NumZeroAllocs );
-
-	// Dump problematic allocations
-	DuplicatedAllocMap.ValueSort( FAllocationInfoSizeGreater() );
-	//FreeWithoutAllocMap
-
-	uint64 TotalDuplicatedMemory = 0;
-	for( const auto& It : DuplicatedAllocMap )
-	{
-		const FAllocationInfo& Alloc = It.Value;
-		TotalDuplicatedMemory += Alloc.Size;
-	}
-
-	UE_LOG( LogStats, Warning, TEXT( "Dumping duplicated alloc map" ) );
-	const float MaxPctDisplayed = 0.80f;
-	uint64 DisplayedSoFar = 0;
-	for( const auto& It : DuplicatedAllocMap )
-	{
-		const FAllocationInfo& Alloc = It.Value;
-		const FString AllocCallstack = FStatsCallstack::GetHumanReadable( Alloc.EncodedCallstack );
-		UE_LOG( LogStats, Log, TEXT( "%lli (%.2f MB) %s" ), Alloc.Size, Alloc.Size / 1024.0f / 1024.0f, *AllocCallstack );
-
-		DisplayedSoFar += Alloc.Size;
-
-		const float CurrentPct = (float)DisplayedSoFar / (float)TotalDuplicatedMemory;
-		if( CurrentPct > MaxPctDisplayed )
-		{
-			break;
-		}
-	}
-
-	InFile->SetProcessingStage( EStatsProcessingStage::SPS_Finished );
-}
-
-void FStatsMemoryDumpCommand::GenerateScopedTreeAllocations( const TMap<FName, FCombinedAllocationInfo>& ScopedAllocations, FNodeAllocationInfo& out_Root )
-{
-	// Experimental code, partially optimized.
-	// Decode all scoped allocations.
-	// Generate tree for allocations and combine them.
-	
-	for (const auto& It : ScopedAllocations)
-	{
-		const FName& EncodedCallstack = It.Key;
-		const FCombinedAllocationInfo& CombinedAllocation = It.Value;
-
-		// Decode callstack.
-		TArray<FName> DecodedCallstack;
-		FStatsCallstack::DecodeToNames( EncodedCallstack, DecodedCallstack );
-		
-		const int32 AllocationLenght = DecodedCallstack.Num();
-		check( DecodedCallstack.Num() > 0 );
-
-		FNodeAllocationInfo* CurrentNode = &out_Root;
-		// Accumulate with thread root node.
-		CurrentNode->Accumulate( CombinedAllocation );
-
-		// Iterate through the callstack and prepare all nodes if needed, and accumulate memory.
-		TArray<FName> CurrentCallstack;
-		const int32 NumEntries = DecodedCallstack.Num();
-		for (int32 Idx1 = 0; Idx1 < NumEntries; ++Idx1)
-		{
-			const FName NodeName = DecodedCallstack[Idx1];
-			CurrentCallstack.Add( NodeName );
-
-			FNodeAllocationInfo* Node = nullptr;
-			const bool bContainsNode = CurrentNode->ChildNodes.Contains( NodeName );
-			if (!bContainsNode)
-			{
-				Node = new FNodeAllocationInfo;
-				Node->Depth = Idx1;
-				Node->PrepareCallstackData( CurrentCallstack );
-
-				CurrentNode->ChildNodes.Add( NodeName, Node );
-			}
-			else
-			{
-				Node = CurrentNode->ChildNodes.FindChecked( NodeName );
-			}
-
-			// Accumulate memory usage and num allocations for all nodes in the callstack.
-			Node->Accumulate( CombinedAllocation );
-			
-			// Move to the next node.
-			Node->Parent = CurrentNode;
-			CurrentNode = Node;
-		}
-	}
-
-	out_Root.SortBySize();
-}
-
-
-void FStatsMemoryDumpCommand::GenerateMemoryUsageReport( const TMap<uint64, FAllocationInfo>& AllocationMap, const TSet<FName>& UObjectRawNames )
-{
-	if( AllocationMap.Num() == 0 )
-	{
-		UE_LOG( LogStats, Warning, TEXT( "There are no allocations, make sure memory profiler is enabled" ) );
-	}
-	else
-	{
-		//ProcessAndDumpUObjectAllocations( AllocationMap );
-		ProcessAndDumpScopedAllocations( AllocationMap );
-	}
-}
-
-void FStatsMemoryDumpCommand::ProcessAndDumpScopedAllocations( const TMap<uint64, FAllocationInfo>& AllocationMap )
-{
-	// This code is not optimized. 
-	FScopeLogTime SLT( TEXT( "ProcessingScopedAllocations" ), nullptr, FScopeLogTime::ScopeLog_Seconds );
-	UE_LOG( LogStats, Warning, TEXT( "Processing scoped allocations" ) );
-
-	const FString ReportName = FString::Printf( TEXT( "%s-Memory-Scoped" ), *GetPlatformName() );
-	FDiagnosticTableViewer MemoryReport( *FDiagnosticTableViewer::GetUniqueTemporaryFilePath( *ReportName ), true );
-
-	// Write a row of headings for the table's columns.
-	MemoryReport.AddColumn( TEXT( "Size (bytes)" ) );
-	MemoryReport.AddColumn( TEXT( "Size (MB)" ) );
-	MemoryReport.AddColumn( TEXT( "Count" ) );
-	MemoryReport.AddColumn( TEXT( "Callstack" ) );
-	MemoryReport.CycleRow();
-
-	TMap<FName, FCombinedAllocationInfo> CombinedAllocations;
-	uint64 TotalAllocatedMemory = 0;
-	uint64 NumAllocations = 0;
-	GenerateScopedAllocations( AllocationMap, CombinedAllocations, TotalAllocatedMemory, NumAllocations );
-
-	// Dump memory to the log.
-	CombinedAllocations.ValueSort( FCombinedAllocationInfoSizeGreater() );
-
-	const float MaxPctDisplayed = 0.90f;
-	int32 CurrentIndex = 0;
-	uint64 DisplayedSoFar = 0;
-	UE_LOG( LogStats, Warning, TEXT( "Index, Size (Size MB), Count, Stat desc" ) );
-	for( const auto& It : CombinedAllocations )
-	{
-		const FCombinedAllocationInfo& CombinedAllocation = It.Value;
-		const FName& EncodedCallstack = It.Key;
-
-		const FString AllocCallstack = FStatsCallstack::GetHumanReadable( EncodedCallstack );
-
-		UE_LOG( LogStats, Log, TEXT( "%2i, %llu (%.2f MB), %llu, %s" ),
-				CurrentIndex,
-				CombinedAllocation.Size,
-				CombinedAllocation.Size / 1024.0f / 1024.0f,
-				CombinedAllocation.Count,
-				*AllocCallstack );
-
-		// Dump stats
-		MemoryReport.AddColumn( TEXT( "%llu" ), CombinedAllocation.Size );
-		MemoryReport.AddColumn( TEXT( "%.2f MB" ), CombinedAllocation.Size / 1024.0f / 1024.0f );
-		MemoryReport.AddColumn( TEXT( "%llu" ), CombinedAllocation.Count );
-		MemoryReport.AddColumn( *AllocCallstack );
-		MemoryReport.CycleRow();
-
-		CurrentIndex++;
-		DisplayedSoFar += CombinedAllocation.Size;
-
-		const float CurrentPct = (float)DisplayedSoFar / (float)TotalAllocatedMemory;
-		if( CurrentPct > MaxPctDisplayed )
-		{
-			break;
-		}
-	}
-
-	UE_LOG( LogStats, Warning, TEXT( "Allocated memory: %llu bytes (%.2f MB)" ), TotalAllocatedMemory, TotalAllocatedMemory / 1024.0f / 1024.0f );
-
-	// Add a total row.
-	MemoryReport.CycleRow();
-	MemoryReport.CycleRow();
-	MemoryReport.CycleRow();
-	MemoryReport.AddColumn( TEXT( "%llu" ), TotalAllocatedMemory );
-	MemoryReport.AddColumn( TEXT( "%.2f MB" ), TotalAllocatedMemory / 1024.0f / 1024.0f );
-	MemoryReport.AddColumn( TEXT( "%llu" ), NumAllocations );
-	MemoryReport.AddColumn( TEXT( "TOTAL" ) );
-	MemoryReport.CycleRow();
-}
-
-void FStatsMemoryDumpCommand::ProcessAndDumpUObjectAllocations( const TMap<uint64, FAllocationInfo>& AllocationMap, const TSet<FName>& UObjectRawNames )
-{
-	// This code is not optimized. 
-	FScopeLogTime SLT( TEXT( "ProcessingUObjectAllocations" ), nullptr, FScopeLogTime::ScopeLog_Seconds );
-	UE_LOG( LogStats, Warning, TEXT( "Processing UObject allocations" ) );
-
-	const FString ReportName = FString::Printf( TEXT( "%s-Memory-UObject" ), *GetPlatformName() );
-	FDiagnosticTableViewer MemoryReport( *FDiagnosticTableViewer::GetUniqueTemporaryFilePath( *ReportName ), true );
-
-	// Write a row of headings for the table's columns.
-	MemoryReport.AddColumn( TEXT( "Size (bytes)" ) );
-	MemoryReport.AddColumn( TEXT( "Size (MB)" ) );
-	MemoryReport.AddColumn( TEXT( "Count" ) );
-	MemoryReport.AddColumn( TEXT( "UObject class" ) );
-	MemoryReport.CycleRow();
-
-	TMap<FName, FCombinedAllocationInfo> UObjectAllocations;
-
-	// To minimize number of calls to expensive DecodeCallstack.
-	TMap<FName,FName> UObjectCallstackToClassMapping;
-
-	uint64 NumAllocations = 0;
-	uint64 TotalAllocatedMemory = 0;
-	for( const auto& It : AllocationMap )
-	{
-		const FAllocationInfo& Alloc = It.Value;
-
-		FName UObjectClass = UObjectCallstackToClassMapping.FindRef( Alloc.EncodedCallstack );
-		if( UObjectClass == NAME_None )
-		{
-			TArray<FName> DecodedCallstack;
-			FStatsCallstack::DecodeToNames( Alloc.EncodedCallstack, DecodedCallstack );
-
-			for( int32 Index = DecodedCallstack.Num() - 1; Index >= 0; --Index )
-			{
-				const FName LongName = DecodedCallstack[Index];
-				const bool bValid = UObjectRawNames.Contains( LongName );
-				if( bValid )
-				{
-					const FString ObjectName = FStatNameAndInfo::GetShortNameFrom( LongName ).GetPlainNameString();
-					UObjectClass = *ObjectName.Left( ObjectName.Find( TEXT( "//" ) ) );;
-					UObjectCallstackToClassMapping.Add( Alloc.EncodedCallstack, UObjectClass );
-					break;
-				}
-			}
-		}
-		
-		if( UObjectClass != NAME_None )
-		{
-			FCombinedAllocationInfo& CombinedAllocation = UObjectAllocations.FindOrAdd( UObjectClass );
-			CombinedAllocation += Alloc;
-
-			TotalAllocatedMemory += Alloc.Size;
-			NumAllocations++;
-		}
-	}
-
-	// Dump memory to the log.
-	UObjectAllocations.ValueSort( FCombinedAllocationInfoSizeGreater() );
-
-	const float MaxPctDisplayed = 0.90f;
-	int32 CurrentIndex = 0;
-	uint64 DisplayedSoFar = 0;
-	UE_LOG( LogStats, Warning, TEXT( "Index, Size (Size MB), Count, UObject class" ) );
-	for( const auto& It : UObjectAllocations )
-	{
-		const FCombinedAllocationInfo& CombinedAllocation = It.Value;
-		const FName& UObjectClass = It.Key;
-
-		UE_LOG( LogStats, Log, TEXT( "%2i, %llu (%.2f MB), %llu, %s" ),
-				CurrentIndex,
-				CombinedAllocation.Size,
-				CombinedAllocation.Size / 1024.0f / 1024.0f,
-				CombinedAllocation.Count,
-				*UObjectClass.GetPlainNameString() );
-
-		// Dump stats
-		MemoryReport.AddColumn( TEXT( "%llu" ), CombinedAllocation.Size );
-		MemoryReport.AddColumn( TEXT( "%.2f MB" ), CombinedAllocation.Size / 1024.0f / 1024.0f );
-		MemoryReport.AddColumn( TEXT( "%llu" ), CombinedAllocation.Count );
-		MemoryReport.AddColumn( *UObjectClass.GetPlainNameString() );
-		MemoryReport.CycleRow();
-
-		CurrentIndex++;
-		DisplayedSoFar += CombinedAllocation.Size;
-
-		const float CurrentPct = (float)DisplayedSoFar / (float)TotalAllocatedMemory;
-		if( CurrentPct > MaxPctDisplayed )
-		{
-			break;
-		}
-	}
-
-	UE_LOG( LogStats, Warning, TEXT( "Allocated memory: %llu bytes (%.2f MB)" ), TotalAllocatedMemory, TotalAllocatedMemory / 1024.0f / 1024.0f );
-
-	// Add a total row.
-	MemoryReport.CycleRow();
-	MemoryReport.CycleRow();
-	MemoryReport.CycleRow();
-	MemoryReport.AddColumn( TEXT( "%llu" ), TotalAllocatedMemory );
-	MemoryReport.AddColumn( TEXT( "%.2f MB" ), TotalAllocatedMemory / 1024.0f / 1024.0f );
-	MemoryReport.AddColumn( TEXT( "%llu" ), NumAllocations );
-	MemoryReport.AddColumn( TEXT( "TOTAL" ) );
-	MemoryReport.CycleRow();
-}
-
-void FStatsMemoryDumpCommand::DumpScopedAllocations( const TCHAR* Name, const TMap<FString, FCombinedAllocationInfo>& CombinedAllocations )
-{
-	if (CombinedAllocations.Num() == 0)
-	{
-		UE_LOG( LogStats, Warning, TEXT( "No scoped allocations: %s" ), Name );
-		return;
-	}
-
-	// This code is not optimized. 
-	FScopeLogTime SLT( TEXT( "ProcessingScopedAllocations" ), nullptr, FScopeLogTime::ScopeLog_Seconds );
-	UE_LOG( LogStats, Warning, TEXT( "Dumping scoped allocations: %s" ), Name );
-
-	const FString ReportName = FString::Printf( TEXT( "%s-Memory-Scoped-%s" ), *GetPlatformName(), Name );
-	FDiagnosticTableViewer MemoryReport( *FDiagnosticTableViewer::GetUniqueTemporaryFilePath( *ReportName ), true );
-
-	// Write a row of headings for the table's columns.
-	MemoryReport.AddColumn( TEXT( "Size (bytes)" ) );
-	MemoryReport.AddColumn( TEXT( "Size (MB)" ) );
-	MemoryReport.AddColumn( TEXT( "Count" ) );
-	MemoryReport.AddColumn( TEXT( "Callstack" ) );
-	MemoryReport.CycleRow();
-
-	
-	FCombinedAllocationInfo Total;
-
-	const float MaxPctDisplayed = 0.90f;
-	int32 CurrentIndex = 0;
-	UE_LOG( LogStats, Warning, TEXT( "Index, Size (Size MB), Count, Stat desc" ) );
-	for (const auto& It : CombinedAllocations)
-	{
-		const FCombinedAllocationInfo& CombinedAllocation = It.Value;
-		//const FName& EncodedCallstack = It.Key;
-		const FString AllocCallstack = It.Key;// GetCallstack( EncodedCallstack );
-
-		UE_LOG( LogStats, Log, TEXT( "%2i, %llu (%.2f MB), %llu, %s" ),
-				CurrentIndex,
-				CombinedAllocation.Size,
-				CombinedAllocation.Size / 1024.0f / 1024.0f,
-				CombinedAllocation.Count,
-				*AllocCallstack );
-
-		// Dump stats
-		MemoryReport.AddColumn( TEXT( "%llu" ), CombinedAllocation.Size );
-		MemoryReport.AddColumn( TEXT( "%.2f MB" ), CombinedAllocation.Size / 1024.0f / 1024.0f );
-		MemoryReport.AddColumn( TEXT( "%llu" ), CombinedAllocation.Count );
-		MemoryReport.AddColumn( *AllocCallstack );
-		MemoryReport.CycleRow();
-
-		CurrentIndex++;
-		Total += CombinedAllocation;
-	}
-
-	UE_LOG( LogStats, Warning, TEXT( "Allocated memory: %llu bytes (%.2f MB)" ), Total.Size, Total.SizeMB );
-
-	// Add a total row.
-	MemoryReport.CycleRow();
-	MemoryReport.CycleRow();
-	MemoryReport.CycleRow();
-	MemoryReport.AddColumn( TEXT( "%llu" ), Total.Size );
-	MemoryReport.AddColumn( TEXT( "%.2f MB" ), Total.SizeMB );
-	MemoryReport.AddColumn( TEXT( "%llu" ), Total.Count );
-	MemoryReport.AddColumn( TEXT( "TOTAL" ) );
-	MemoryReport.CycleRow();
-}
-
-void FStatsMemoryDumpCommand::GenerateScopedAllocations( const TMap<uint64, FAllocationInfo>& AllocationMap, TMap<FName, FCombinedAllocationInfo>& out_CombinedAllocations, uint64& TotalAllocatedMemory, uint64& NumAllocations )
-{
-	FScopeLogTime SLT( TEXT( "GenerateScopedAllocations" ), nullptr, FScopeLogTime::ScopeLog_Milliseconds );
-
-	for (const auto& It : AllocationMap)
-	{
-		const FAllocationInfo& Alloc = It.Value;
-		FCombinedAllocationInfo& CombinedAllocation = out_CombinedAllocations.FindOrAdd( Alloc.EncodedCallstack );
-		CombinedAllocation += Alloc;
-
-		TotalAllocatedMemory += Alloc.Size;
-		NumAllocations++;
-	}
-
-	// Sort by size.
-	out_CombinedAllocations.ValueSort( FCombinedAllocationInfoSizeGreater() );
-}
-
-void FStatsMemoryDumpCommand::PrepareSnapshot( const FName SnapshotName, const TMap<uint64, FAllocationInfo>& AllocationMap )
-{
-	FScopeLogTime SLT( TEXT( "PrepareSnapshot" ), nullptr, FScopeLogTime::ScopeLog_Milliseconds );
-
-	// Make sure the snapshot name is unique.
-	FName UniqueSnapshotName = SnapshotName;
-	while (SnapshotNames.Contains( UniqueSnapshotName ))
-	{
-		UniqueSnapshotName = FName( UniqueSnapshotName, UniqueSnapshotName.GetNumber() + 1 );
-	}
-	SnapshotNames.Add( UniqueSnapshotName );
-
-	SnapshotsWithAllocationMap.Add( UniqueSnapshotName, AllocationMap );
-
-	TMap<FName, FCombinedAllocationInfo> SnapshotCombinedAllocations;
-	uint64 TotalAllocatedMemory = 0;
-	uint64 NumAllocations = 0;
-	GenerateScopedAllocations( AllocationMap, SnapshotCombinedAllocations, TotalAllocatedMemory, NumAllocations );
-	SnapshotsWithScopedAllocations.Add( UniqueSnapshotName, SnapshotCombinedAllocations );
-
-	// Decode callstacks.
-	// Replace encoded callstacks with human readable name. For easier debugging.
-	TMap<FString, FCombinedAllocationInfo> SnapshotDecodedCombinedAllocations;
-	for (auto& It : SnapshotCombinedAllocations)
-	{
-		const FString HumanReadableCallstack = FStatsCallstack::GetHumanReadable( It.Key );
-		SnapshotDecodedCombinedAllocations.Add( HumanReadableCallstack, It.Value );
-	}
-	SnapshotsWithDecodedScopedAllocations.Add( UniqueSnapshotName, SnapshotDecodedCombinedAllocations );
-
-	UE_LOG( LogStats, Warning, TEXT( "PrepareSnapshot: %s Alloc: %i Scoped: %i Total: %.2f MB" ), *UniqueSnapshotName.ToString(), AllocationMap.Num(), SnapshotCombinedAllocations.Num(), TotalAllocatedMemory / 1024.0f / 1024.0f );
-}
-
-void FStatsMemoryDumpCommand::CompareSnapshots_FName( const FName BeginSnaphotName, const FName EndSnaphotName, TMap<FName, FCombinedAllocationInfo>& out_Result )
-{
-	const auto BeginSnaphotPtr = SnapshotsWithScopedAllocations.Find( BeginSnaphotName );
-	const auto EndSnapshotPtr = SnapshotsWithScopedAllocations.Find( EndSnaphotName );
-	if (BeginSnaphotPtr && EndSnapshotPtr)
-	{
-		// Process data.
-		TMap<FName, FCombinedAllocationInfo> BeginSnaphot = *BeginSnaphotPtr;
-		TMap<FName, FCombinedAllocationInfo> EndSnaphot = *EndSnapshotPtr;
-		TMap<FName, FCombinedAllocationInfo> Result;
-
-		for (const auto& It : EndSnaphot)
-		{
-			const FName Callstack = It.Key;
-			const FCombinedAllocationInfo EndCombinedAlloc = It.Value;
-
-			const FCombinedAllocationInfo* BeginCombinedAllocPtr = BeginSnaphot.Find( Callstack );
-			if (BeginCombinedAllocPtr)
-			{
-				FCombinedAllocationInfo CombinedAllocation;
-				CombinedAllocation += EndCombinedAlloc;
-				CombinedAllocation -= *BeginCombinedAllocPtr;
-
-				if (CombinedAllocation.IsAlive())
-				{
-					out_Result.Add( Callstack, CombinedAllocation );
-				}
-			}
-			else
-			{
-				out_Result.Add( Callstack, EndCombinedAlloc );
-			}
-		}
-
-		// Sort by size.
-		out_Result.ValueSort( FCombinedAllocationInfoSizeGreater() );
-	}
-}
-
-void FStatsMemoryDumpCommand::CompareSnapshots_FString( const FName BeginSnaphotName, const FName EndSnaphotName, TMap<FString, FCombinedAllocationInfo>& out_Result )
-{
-	const auto BeginSnaphotPtr = SnapshotsWithDecodedScopedAllocations.Find( BeginSnaphotName );
-	const auto EndSnapshotPtr = SnapshotsWithDecodedScopedAllocations.Find( EndSnaphotName );
-	if (BeginSnaphotPtr && EndSnapshotPtr)
-	{
-		// Process data.
-		TMap<FString, FCombinedAllocationInfo> BeginSnaphot = *BeginSnaphotPtr;
-		TMap<FString, FCombinedAllocationInfo> EndSnaphot = *EndSnapshotPtr;
-
-		for (const auto& It : EndSnaphot)
-		{
-			const FString& Callstack = It.Key;
-			const FCombinedAllocationInfo EndCombinedAlloc = It.Value;
-
-			const FCombinedAllocationInfo* BeginCombinedAllocPtr = BeginSnaphot.Find( Callstack );
-			if (BeginCombinedAllocPtr)
-			{
-				FCombinedAllocationInfo CombinedAllocation;
-				CombinedAllocation += EndCombinedAlloc;
-				CombinedAllocation -= *BeginCombinedAllocPtr;
-
-				if (CombinedAllocation.IsAlive())
-				{
-					out_Result.Add( Callstack, CombinedAllocation );
-				}
-			}
-			else
-			{
-				out_Result.Add( Callstack, EndCombinedAlloc );
-			}
-		}
-
-		// Sort by size.
-		out_Result.ValueSort( FCombinedAllocationInfoSizeGreater() );
-	}
+	const FName NAME_ProfilerModule = TEXT( "Profiler" );
+	IProfilerModule& ProfilerModule = FModuleManager::LoadModuleChecked<IProfilerModule>( NAME_ProfilerModule );
+	ProfilerModule.StatsMemoryDumpCommand( *SourceFilepath );
 }
 
diff --git a/Engine/Source/Programs/UnrealFrontend/Private/Commands/StatsDumpMemoryCommand.h b/Engine/Source/Programs/UnrealFrontend/Private/Commands/StatsDumpMemoryCommand.h
--- a/Engine/Source/Programs/UnrealFrontend/Private/Commands/StatsDumpMemoryCommand.h
+++ b/Engine/Source/Programs/UnrealFrontend/Private/Commands/StatsDumpMemoryCommand.h
@@ -2,278 +2,11 @@
 
 #pragma once
 
-#include "StatsData.h"
-#include "StatsFile.h"
-
-/*-----------------------------------------------------------------------------
-	Basic structures
------------------------------------------------------------------------------*/
-
-/** Simple allocation info. Assumes the sequence tag are unique and doesn't wrap. */
-struct FAllocationInfo
-{
-	/** Old pointer, for Realloc. */
-	uint64 OldPtr;
-
-	/** Pointer, for Alloc, Realloc and Free. */
-	uint64 Ptr;
-
-	/** Size of the allocation. */
-	int64 Size;
-	
-	/**
-	 * Encoded callstack as FNames indices separated with +
-	 * 123+45+56+76
-	 */
-	FName EncodedCallstack;
-	
-	/**
-	 * Index of the memory operation. 
-	 * Stats are handled with packets sent from different threads.
-	 * This is used to sort the memory operations to be in the sequence order.
-	 */
-	uint32 SequenceTag;
-	
-	/** Memory operation, can be Alloc, Free or Realloc. */
-	EMemoryOperation Op;
-	
-	/** If true, the callstack is broken. Missing Start/End scope. */
-	bool bHasBrokenCallstack;
-
-	/** Initialization constructor. */
-	FAllocationInfo(
-		uint64 InOldPtr,
-		uint64 InPtr,
-		int64 InSize,
-		const TArray<FName>& InCallstack,
-		uint32 InSequenceTag,
-		EMemoryOperation InOp,
-		bool bInHasBrokenCallstack
-		);
-
-	/** Copy constructor. */
-	FAllocationInfo( const FAllocationInfo& Other );
-};
-
-/** Contains data about combined allocations for an unique scope. */
-struct FCombinedAllocationInfo
-{
-	/** Size as MB. */
-	double SizeMB;
-
-	/** Allocations count. */
-	int64 Count;
-
-	/** Human readable callstack like 'GameThread -> Preinit -> LoadModule */
-	FString HumanReadableCallstack;
-
-	/** Encoded callstack like '45+656+6565'. */
-	FName EncodedCallstack;
-
-	/** Callstack indices, encoded callstack parsed into an array [45,656,6565] */
-	TArray<FName> DecodedCallstack;
-
-	/** Size for this node and children. */
-	int64 Size;
-
-	/** Default constructor. */
-	FCombinedAllocationInfo()
-		: SizeMB( 0 )
-		, Count( 0 )
-		, Size( 0 )
-	{}
-
-	/** Operator for adding another combined allocation. */
-	FCombinedAllocationInfo& operator+=(const FCombinedAllocationInfo& Other)
-	{
-		if ((UPTRINT*)this != (UPTRINT*)&Other)
-		{
-			AddAllocation( Other.Size, Other.Count );
-		}
-		return *this;
-	}
-
-	/** Operator for subtracting another combined allocation. */
-	FCombinedAllocationInfo& operator-=(const FCombinedAllocationInfo& Other)
-	{
-		if ((UPTRINT*)this != (UPTRINT*)&Other)
-		{
-			AddAllocation( -Other.Size, -Other.Count );
-		}
-		return *this;
-	}
-
-	/** Operator for adding a new allocation. */
-	FCombinedAllocationInfo& operator+=(const FAllocationInfo& Other)
-	{
-		if ((UPTRINT*)this != (UPTRINT*)&Other)
-		{
-			AddAllocation( Other.Size, 1 );
-		}
-		return *this;
-	}
-	
-	/**
-	 * @return true, if the combined allocation is alive, the allocated data is larger than 0
-	 */
-	bool IsAlive()
-	{
-		return Size > 0;
-	}
-
-protected:
-	/** Adds a new allocation. */
-	void AddAllocation( int64 SizeToAdd, int64 CountToAdd = 1 )
-	{
-		Size += SizeToAdd;
-		Count += CountToAdd;
-
-		const double InvMB = 1.0 / 1024.0 / 1024.0;
-		SizeMB = Size * InvMB;
-	}
-};
-
-/** Contains data about node allocations, parent and children allocations. */
-struct FNodeAllocationInfo
-{
-	/** Size as MBs for this node and children. */
-	double SizeMB;
-
-	/** Allocations count for this node and children. */
-	int64 Count;
-
-	/** Human readable callstack like 'GameThread -> Preinit -> LoadModule */
-	FString HumanReadableCallstack;
-
-	/** Encoded callstack like '45+656+6565'. */
-	FName EncodedCallstack;
-
-	/** Callstack indices, encoded callstack parsed into an array [45,656,6565] */
-	TArray<FName> DecodedCallstack;
-
-	/** Parent node. */
-	FNodeAllocationInfo* Parent;
-
-	/** Child nodes. */
-	TMap<FName, FNodeAllocationInfo*> ChildNodes;
-
-	/** Size for this node and children. */
-	int64 Size;
-
-	/** Node's depth. */
-	int32 Depth;
-
-	/** Default constructor. */
-	FNodeAllocationInfo()
-		: SizeMB( 0.0 )
-		, Count( 0 )
-		, Parent( nullptr )
-		, Size( 0 )
-		, Depth( 0 )
-	{}
-
-	/** Destructor. */
-	~FNodeAllocationInfo()
-	{
-		DeleteAllChildrenNodes();
-	}
-
-	/** Accumulates this node with the other combined allocation. */
-	void Accumulate( const FCombinedAllocationInfo& Other )
-	{
-		Size += Other.Size;
-		Count += Other.Count;
-
-		const double InvMB = 1.0 / 1024.0 / 1024.0;
-		SizeMB = Size * InvMB;
-	}
-
-	/** Recursively sorts all children by size. */
-	void SortBySize();
-
-	/** Prepares all callstack data for this node. */
-	void PrepareCallstackData( const TArray<FName>& InDecodedCallstack );
-
-protected:
-	/** Recursively deletes all children. */
-	void DeleteAllChildrenNodes()
-	{
-		for (auto& It : ChildNodes)
-		{
-			delete It.Value;
-		}
-		ChildNodes.Empty();
-	}
-};
-
-// #YRX_Profiler: 2015-07-10 Should be moved to the profiler module, and added as export or something?
-/** An example command loading a raw stats file and dumping memory usage. */
+/** An example command loading a raw stats file and dumping memory usage. Mostly a dummy class that will execute required code from the profiler module. */
 class FStatsMemoryDumpCommand
 {
 public:
 
 	/** Executes the command. */
-	static void Run()
-	{
-		FStatsMemoryDumpCommand Instance;
-		Instance.InternalRun();
-	}
-
-protected:
-	/** Executes the command. */
-	void InternalRun();
-
-	/** Basic memory profiling, only for debugging purpose. */
-	void ProcessMemoryOperations( FStatsReadFile* InFile );
-
-	/** Generates a basic memory usage report and prints it to the log. */
-	void GenerateMemoryUsageReport( const TMap<uint64, FAllocationInfo>& AllocationMap, const TSet<FName>& UObjectRawNames );
-
-	/** Generates scoped allocation statistics. */
-	void ProcessAndDumpScopedAllocations( const TMap<uint64, FAllocationInfo>& AllocationMap );
-
-	/** Generates UObject allocation statistics. */
-	void ProcessAndDumpUObjectAllocations( const TMap<uint64, FAllocationInfo>& AllocationMap, const TSet<FName>& UObjectRawNames );
-
-	void DumpScopedAllocations( const TCHAR* Name, const TMap<FString, FCombinedAllocationInfo>& CombinedAllocations );
-
-	/** Generates callstack based allocation map. */
-	void GenerateScopedAllocations( const TMap<uint64, FAllocationInfo>& AllocationMap, TMap<FName, FCombinedAllocationInfo>& out_CombinedAllocations, uint64& TotalAllocatedMemory, uint64& NumAllocations );
-
-	void GenerateScopedTreeAllocations( const TMap<FName, FCombinedAllocationInfo>& ScopedAllocations, FNodeAllocationInfo& out_Root );
-
-	/** Prepares data for a snapshot. */
-	void PrepareSnapshot( const FName SnapshotName, const TMap<uint64, FAllocationInfo>& AllocationMap );
-
-	void CompareSnapshots_FName( const FName BeginSnaphotName, const FName EndSnaphotName, TMap<FName, FCombinedAllocationInfo>& out_Result );
-	void CompareSnapshots_FString( const FName BeginSnaphotName, const FName EndSnaphotName, TMap<FString, FCombinedAllocationInfo>& out_Result );
-
-
-	/**
-	 * @return Platform's name based on the loaded ue4statsraw file
-	 */
-	const FString& GetPlatformName() const
-	{
-		return PlatformName;
-	}
-
-	/** The sequence tag mapping to the named markers. */
-	TArray<TPair<uint32, FName>> Snapshots;
-
-	/** Unique snapshot names. */
-	TSet<FName> SnapshotNames;
-
-	/** The sequence tag mapping to the named markers that need to processed. */
-	TArray<TPair<uint32, FName>> SnapshotsToBeProcessed;
-
-	/** Snapshots with allocation map. */
-	TMap<FName, TMap<uint64, FAllocationInfo> > SnapshotsWithAllocationMap;
-
-	/** Snapshots with callstack based allocation map. */
-	TMap<FName, TMap<FName, FCombinedAllocationInfo> > SnapshotsWithScopedAllocations;
-
-	TMap<FName, TMap<FString, FCombinedAllocationInfo> > SnapshotsWithDecodedScopedAllocations;
-
-	/** Platform's name based on the loaded ue4statsraw file. */
-	FString PlatformName;
+	static void Run();
 };
diff --git a/Engine/Source/Runtime/Core/Private/Stats/StatsFile.cpp b/Engine/Source/Runtime/Core/Private/Stats/StatsFile.cpp
--- a/Engine/Source/Runtime/Core/Private/Stats/StatsFile.cpp
+++ b/Engine/Source/Runtime/Core/Private/Stats/StatsFile.cpp
@@ -548,15 +548,11 @@
 
 const double FStatsReadFile::NumSecondsBetweenUpdates = 5.0;
 
-FStatsReadFile* FStatsReadFile::CreateReaderForRegularStats( const TCHAR* Filename, FOnNewStatPacket InNewStatPacket )
+FStatsReadFile* FStatsReadFile::CreateReaderForRegularStats( const TCHAR* Filename )
 {
 	FStatsReadFile* StatsReadFile = new FStatsReadFile( Filename, false );
 	const bool bValid = StatsReadFile->PrepareLoading();
-	if (bValid)
+	if (!bValid)
 	{
-		StatsReadFile->NewStatPacket = InNewStatPacket;
-	}
-	else
-	{
 		delete StatsReadFile;
 	}
@@ -564,30 +560,24 @@
 }
 
 
-FStatsReadFile* FStatsReadFile::CreateReaderForRawStats( const TCHAR* Filename, FOnNewCombinedHistory InNewCombinedHistory )
-{
-	FStatsReadFile* StatsReadFile = new FStatsReadFile( Filename, true );
-	const bool bValid = StatsReadFile->PrepareLoading();
-	if (bValid)
-	{
-		StatsReadFile->NewCombinedHistory = InNewCombinedHistory;
-	}
-	else
-	{
-		delete StatsReadFile;
-	}
-	return bValid ? StatsReadFile : nullptr;
-}
-
 void FStatsReadFile::ReadAndProcessSynchronously()
 {
+	// Read.
 	ReadStats();
+
+	// Process.
+	PreProcessStats();
 	ProcessStats();
+	PostProcessStats();
 
 	if (GetProcessingStage() == EStatsProcessingStage::SPS_Stopped)
 	{
 		SetProcessingStage( EStatsProcessingStage::SPS_Invalid );
 	}
+	else
+	{
+		SetProcessingStage( EStatsProcessingStage::SPS_Finished );
+	}
 }
 
 void FStatsReadFile::ReadAndProcessAsynchronously()
@@ -599,110 +589,6 @@
 	}
 }
 
-void FStatsReadFile::ReadStats()
-{
-	if (bRawStatsFile)
-	{
-		const double StartTime = FPlatformTime::Seconds();
-
-		// Buffer used to store the compressed and decompressed data.
-		TArray<uint8> SrcArray;
-		TArray<uint8> DestArray;
-
-		// Read all packets sequentially, forced by the memory profiler which is now a part of the raw stats.
-		// !!CAUTION!! Frame number in the raw stats is pointless, because it is time/cycles based, not frame based.
-		// Background threads usually execute time consuming operations, so the frame number won't be valid.
-		// Needs to be combined by the thread and the time, not by the frame number.
-
-		// Update stage progress once per NumSecondsBetweenUpdates(5) seconds to avoid spamming.
-		double PreviousSeconds = FPlatformTime::Seconds();
-		SetProcessingStage( EStatsProcessingStage::SPS_ReadAndCombinePackets );
-
-		while (Reader->Tell() < Reader->TotalSize())
-		{
-			// Read the compressed data.
-			FCompressedStatsData UncompressedData( SrcArray, DestArray );
-			*Reader << UncompressedData;
-			if (UncompressedData.HasReachedEndOfCompressedData())
-			{
-				StageProgress.Set( 100 );
-				break;
-			}
-
-			FMemoryReader MemoryReader( DestArray, true );
-
-			FStatPacket* StatPacket = new FStatPacket();
-			Stream.ReadStatPacket( MemoryReader, *StatPacket );
-
-			const int64 StatPacketFrameNum = StatPacket->Frame;
-			FStatPacketArray& Frame = CombinedHistory.FindOrAdd( StatPacketFrameNum );
-
-			// Check if we need to combine packets from the same thread.
-			FStatPacket** CombinedPacket = Frame.Packets.FindByPredicate( [&]( FStatPacket* Item ) -> bool
-			{
-				return Item->ThreadId == StatPacket->ThreadId;
-			} );
-
-			
-			if (CombinedPacket)
-			{
-				(*CombinedPacket)->StatMessages += StatPacket->StatMessages;
-				delete StatPacket;
-			}
-			else
-			{
-				Frame.Packets.Add( StatPacket );
-				FileInfo.MaximumPacketSize = FMath::Max<int64>( FileInfo.MaximumPacketSize, StatPacket->StatMessages.GetAllocatedSize() );
-			}
-
-			UpdateReadStageProgress();
-
-			if (bShouldStopProcessing == true)
-			{
-				SetProcessingStage( EStatsProcessingStage::SPS_Stopped );
-				break;
-			}
-
-			FileInfo.TotalPacketsNum++;
-		}
-
-		if (GetProcessingStage() != EStatsProcessingStage::SPS_Stopped)
-		{
-			const double TotalTime = FPlatformTime::Seconds() - StartTime;
-			UE_LOG( LogStats, Log, TEXT( "Reading took %.2f sec(s)" ), TotalTime );
-
-			UpdateCombinedHistoryStats();
-		}
-		else
-		{
-			UE_LOG( LogStats, Warning, TEXT( "Reading stopped, abandoning" ) );
-			// Clear all data.
-			CombinedHistory.Empty();
-		}
-	}
-}
-
-void FStatsReadFile::ProcessStats()
-{
-	if (bRawStatsFile)
-	{
-		if (GetProcessingStage() != EStatsProcessingStage::SPS_Stopped)
-		{
-			SetProcessingStage( EStatsProcessingStage::SPS_ProcessStatMessages );
-			NewCombinedHistory.ExecuteIfBound( this );
-
-			if (GetProcessingStage() == EStatsProcessingStage::SPS_Stopped)
-			{
-				UE_LOG( LogStats, Warning, TEXT( "Processing stopped, abandoning" ) );
-			}
-
-			// Clear all data.
-			CombinedHistory.Empty();
-		}
-	}
-}
-
-
 FStatsReadFile::FStatsReadFile( const TCHAR* InFilename, bool bInRawStatsFile )
 	: Header( Stream.Header )
 	, Reader( nullptr )
@@ -711,28 +597,7 @@
 	, Filename( InFilename )
 	, bRawStatsFile( bInRawStatsFile )
 {
-	
+
 }
 
-FStatsReadFile::~FStatsReadFile()
-{
-	RequestStop();
-
-	FPlatformProcess::ConditionalSleep( [&]()
-	{
-		return !IsBusy();
-	}, 1.0f );
-
-	if (AsyncWork)
-	{
-		check( AsyncWork->IsDone() );
-
-		delete AsyncWork;
-		AsyncWork = nullptr;
-	}
-
-	delete Reader;
-	Reader = nullptr;
-}
-
 bool FStatsReadFile::PrepareLoading()
@@ -823,7 +688,321 @@
 	return true;
 }
 
-void FStatsReadFile::UpdateReadStageProgress()
+FStatsReadFile::~FStatsReadFile()
 {
+	RequestStop();
+
+	FPlatformProcess::ConditionalSleep( [&]()
+	{
+		return !IsBusy();
+	}, 1.0f );
+
+	if (AsyncWork)
+	{
+		check( AsyncWork->IsDone() );
+
+		delete AsyncWork;
+		AsyncWork = nullptr;
+	}
+
+	delete Reader;
+	Reader = nullptr;
+}
+
+void FStatsReadFile::ReadStats()
+{
+	if (bRawStatsFile)
+	{
+		const double StartTime = FPlatformTime::Seconds();
+
+		// Buffer used to store the compressed and decompressed data.
+		TArray<uint8> SrcArray;
+		TArray<uint8> DestArray;
+
+		// Read all packets sequentially, forced by the memory profiler which is now a part of the raw stats.
+		// !!CAUTION!! Frame number in the raw stats is pointless, because it is time/cycles based, not frame based.
+		// Background threads usually execute time consuming operations, so the frame number won't be valid.
+		// Needs to be combined by the thread and the time, not by the frame number.
+
+		// Update stage progress once per NumSecondsBetweenUpdates(5) seconds to avoid spamming.
+		SetProcessingStage( EStatsProcessingStage::SPS_ReadStats );
+
+		while (Reader->Tell() < Reader->TotalSize())
+		{
+			// Read the compressed data.
+			FCompressedStatsData UncompressedData( SrcArray, DestArray );
+			*Reader << UncompressedData;
+			if (UncompressedData.HasReachedEndOfCompressedData())
+			{
+				StageProgress.Set( 100 );
+				break;
+			}
+
+			FMemoryReader MemoryReader( DestArray, true );
+
+			FStatPacket* StatPacket = new FStatPacket();
+			Stream.ReadStatPacket( MemoryReader, *StatPacket );
+
+			const int32 StatPacketFrameNum = int32(StatPacket->Frame);
+			FStatPacketArray& Frame = CombinedHistory.FindOrAdd( StatPacketFrameNum );
+
+			// Check if we need to combine packets from the same thread.
+			FStatPacket** CombinedPacket = Frame.Packets.FindByPredicate( [&]( FStatPacket* Item ) -> bool
+			{
+				return Item->ThreadId == StatPacket->ThreadId;
+			} );
+			
+			if (CombinedPacket)
+			{
+				(*CombinedPacket)->StatMessages += StatPacket->StatMessages;
+				delete StatPacket;
+			}
+			else
+			{
+				Frame.Packets.Add( StatPacket );
+				FileInfo.MaximumPacketSize = FMath::Max<int32>( FileInfo.MaximumPacketSize, StatPacket->StatMessages.GetAllocatedSize() );
+			}
+
+			if (!UpdateReadStageProgress())
+			{
+				break;
+			}
+			
+			FileInfo.TotalPacketsNum++;
+		}
+
+		// Generate frames array.
+		CombinedHistory.GenerateKeyArray( Frames );
+		Frames.Sort();
+		// Verify that frames are sequential.
+		check( Frames[Frames.Num() - 1] == Frames.Num() );
+
+		if (GetProcessingStage() != EStatsProcessingStage::SPS_Stopped)
+		{
+			const double TotalTime = FPlatformTime::Seconds() - StartTime;
+			UE_LOG( LogStats, Log, TEXT( "Reading took %.2f sec(s)" ), TotalTime );
+
+			UpdateCombinedHistoryStats();
+		}
+		else
+		{
+			UE_LOG( LogStats, Warning, TEXT( "Reading stopped, abandoning" ) );
+			// Clear all data.
+			CombinedHistory.Empty();
+		}
+	}
+}
+
+void FStatsReadFile::PreProcessStats()
+{
+	if (GetProcessingStage() != EStatsProcessingStage::SPS_Stopped)
+	{
+		SetProcessingStage( EStatsProcessingStage::SPS_PreProcessStats );
+	}
+}
+
+void FStatsReadFile::ProcessStats()
+{
+	if (bRawStatsFile)
+	{
+		if (GetProcessingStage() != EStatsProcessingStage::SPS_Stopped)
+		{
+			SetProcessingStage( EStatsProcessingStage::SPS_ProcessStats );
+			const double StartTime = FPlatformTime::Seconds();
+
+			int32 CurrentStatMessageIndex = 0;
+
+			// Raw stats callstack for this file.
+			TMap<FName, FStackState> StackStates;
+
+			// Read all stats messages for all frames, decode callstacks.
+			const int32 FirstFrame = 0;
+			const int32 OnerPercent = FMath::Max( int32( FileInfo.TotalStatMessagesNum / 100 ), 65536 );
+
+			for (int32 FrameIndex = 0; FrameIndex < Frames.Num(); ++FrameIndex)
+			{
+				const int32 TargetFrame = Frames[FrameIndex];
+				const int32 Diff = TargetFrame - FirstFrame;
+				const FStatPacketArray& Frame = CombinedHistory.FindChecked( TargetFrame );
+
+				for (int32 PacketIndex = 0; PacketIndex < Frame.Packets.Num(); PacketIndex++)
+				{
+					const FStatPacket& StatPacket = *Frame.Packets[PacketIndex];
+					const FName& ThreadFName = State.Threads.FindChecked( StatPacket.ThreadId );
+
+					FStackState* StackState = StackStates.Find( ThreadFName );
+					if (!StackState)
+					{
+						StackState = &StackStates.Add( ThreadFName );
+						StackState->Stack.Add( ThreadFName );
+						StackState->Current = ThreadFName;
+					}
+
+					const FStatMessagesArray& Data = StatPacket.StatMessages;
+					const int32 NumStatMessages = Data.Num();
+					for (int32 Index = 0; Index < NumStatMessages; Index++)
+					{
+						CurrentStatMessageIndex++;
+						if (CurrentStatMessageIndex % OnerPercent == 0)
+						{
+							UpdateProcessStageProgress( CurrentStatMessageIndex, FrameIndex, PacketIndex );
+						}
+
+						const FStatMessage& Message = Data[Index];
+						const EStatOperation::Type Op = Message.NameAndInfo.GetField<EStatOperation>();
+						const FName RawName = Message.NameAndInfo.GetRawName();
+
+						if (Op == EStatOperation::CycleScopeStart || Op == EStatOperation::CycleScopeEnd || Op == EStatOperation::Memory || Op == EStatOperation::SpecialMessageMarker)
+						{
+							if (Op == EStatOperation::CycleScopeStart)
+							{
+								StackState->Stack.Add( RawName );
+								StackState->Current = RawName;
+								ProcessCycleScopeStartOperation( Message, *StackState );
+							}
+							else if (Op == EStatOperation::Memory)
+							{
+								// ProcessMemoryOperation
+
+								// Experimental code used only to test the implementation.
+								// First memory operation is Alloc or Free
+								const uint64 EncodedPtr = Message.GetValue_Ptr();
+								const EMemoryOperation MemOp = EMemoryOperation( EncodedPtr & (uint64)EMemoryOperation::Mask );
+								const uint64 Ptr = EncodedPtr & ~(uint64)EMemoryOperation::Mask;
+								if (MemOp == EMemoryOperation::Alloc)
+								{
+									// @see FStatsMallocProfilerProxy::TrackAlloc
+									// After AllocPtr message there is always alloc size message and the sequence tag.
+									Index++; CurrentStatMessageIndex++;
+									const FStatMessage& AllocSizeMessage = Data[Index];
+									const int64 AllocSize = AllocSizeMessage.GetValue_int64();
+
+									// Read OperationSequenceTag.
+									Index++; CurrentStatMessageIndex++;
+									const FStatMessage& SequenceTagMessage = Data[Index];
+									const uint32 SequenceTag = SequenceTagMessage.GetValue_int64();		
+
+									//ThreadStats->AddMemoryMessage( GET_STATFNAME( STAT_Memory_AllocPtr ), (uint64)(UPTRINT)Ptr | (uint64)EMemoryOperation::Alloc );
+									//ThreadStats->AddMemoryMessage( GET_STATFNAME( STAT_Memory_AllocSize ), Size );
+									//ThreadStats->AddMemoryMessage( GET_STATFNAME( STAT_Memory_OperationSequenceTag ), (int64)SequenceTag );
+									ProcessMemoryOperation( MemOp, Ptr, 0, AllocSize, SequenceTag, *StackState );
+								}
+								else if (MemOp == EMemoryOperation::Realloc)
+								{
+									const uint64 OldPtr = Ptr;
+
+									// Read NewPtr
+									Index++; CurrentStatMessageIndex++;
+									const FStatMessage& AllocPtrMessage = Data[Index];
+									const uint64 NewPtr = AllocPtrMessage.GetValue_Ptr() & ~(uint64)EMemoryOperation::Mask;
+
+									// After AllocPtr message there is always alloc size message and the sequence tag.
+									Index++; CurrentStatMessageIndex++;
+									const FStatMessage& ReallocSizeMessage = Data[Index];
+									const int64 ReallocSize = ReallocSizeMessage.GetValue_int64();
+
+									// Read OperationSequenceTag.
+									Index++; CurrentStatMessageIndex++;
+									const FStatMessage& SequenceTagMessage = Data[Index];
+									const uint32 SequenceTag = SequenceTagMessage.GetValue_int64();
+
+									//ThreadStats->AddMemoryMessage( GET_STATFNAME( STAT_Memory_FreePtr ), (uint64)(UPTRINT)OldPtr | (uint64)EMemoryOperation::Realloc );
+									//ThreadStats->AddMemoryMessage( GET_STATFNAME( STAT_Memory_AllocPtr ), (uint64)(UPTRINT)NewPtr | (uint64)EMemoryOperation::Realloc );
+									//ThreadStats->AddMemoryMessage( GET_STATFNAME( STAT_Memory_AllocSize ), NewSize );
+									//ThreadStats->AddMemoryMessage( GET_STATFNAME( STAT_Memory_OperationSequenceTag ), (int64)SequenceTag );
+									ProcessMemoryOperation( MemOp, OldPtr, NewPtr, ReallocSize, SequenceTag, *StackState );
+								}
+								else if (MemOp == EMemoryOperation::Free)
+								{
+									// Read OperationSequenceTag.
+									Index++; CurrentStatMessageIndex++;
+									const FStatMessage& SequenceTagMessage = Data[Index];
+									const uint32 SequenceTag = SequenceTagMessage.GetValue_int64();
+
+									//ThreadStats->AddMemoryMessage( GET_STATFNAME( STAT_Memory_FreePtr ), (uint64)(UPTRINT)Ptr | (uint64)EMemoryOperation::Free );	// 16 bytes total				
+									//ThreadStats->AddMemoryMessage( GET_STATFNAME( STAT_Memory_OperationSequenceTag ), (int64)SequenceTag );
+									ProcessMemoryOperation( MemOp, Ptr, 0, 0, SequenceTag, *StackState );
+								}
+								else
+								{
+									UE_LOG( LogStats, Warning, TEXT( "Pointer from a memory operation is invalid" ) );
+								}
+							}
+							// Set, Clear, Add, Subtract
+							else if (Op == EStatOperation::CycleScopeEnd)
+							{
+								if (StackState->Stack.Num() > 1)
+								{
+									const FName ScopeStart = StackState->Stack.Pop();
+									const FName ScopeEnd = Message.NameAndInfo.GetRawName();
+
+									check( ScopeStart == ScopeEnd );
+
+									StackState->Current = StackState->Stack.Last();
+
+									// The stack should be ok, but it may be partially broken.
+									// This will happen if memory profiling starts in the middle of executing a background thread.
+									StackState->bIsBrokenCallstack = false;
+
+									ProcessCycleScopeEndOperation( Message, *StackState );
+								}
+								else
+								{
+									const FName ShortName = Message.NameAndInfo.GetShortName();
+
+									UE_LOG( LogStats, Warning, TEXT( "Broken cycle scope end %s/%s, current %s" ),
+											*ThreadFName.ToString(),
+											*ShortName.ToString(),
+											*StackState->Current.ToString() );
+
+									// The stack is completely broken, only has the thread name and the last cycle scope.
+									// Rollback to the thread node.
+									StackState->bIsBrokenCallstack = true;
+									StackState->Stack.Empty();
+									StackState->Stack.Add( ThreadFName );
+									StackState->Current = ThreadFName;
+
+									//?ProcessCycleScopeEndOperation( Message, *StackState );
+								}
+							}
+							else if (Op == EStatOperation::SpecialMessageMarker)
+							{
+								ProcessSpecialMessageMarkerOperation( Message, *StackState );
+							}
+						}
+					}
+				}
+			}
+
+			if (GetProcessingStage() != EStatsProcessingStage::SPS_Stopped)
+			{
+				StageProgress.Set( 100 );
+
+				const double TotalTime = FPlatformTime::Seconds() - StartTime;
+				UE_LOG( LogStats, Log, TEXT( "Processing took %.2f sec(s)" ), TotalTime );
+
+				//UpdateCombinedHistoryStats();
+			}
+			else
+			{
+				UE_LOG( LogStats, Warning, TEXT( "Processing stopped, abandoning" ) );
+			}
+
+			// Clear all data.
+			CombinedHistory.Empty();
+		}
+	}
+}
+
+void FStatsReadFile::PostProcessStats()
+{
+	if (GetProcessingStage() != EStatsProcessingStage::SPS_Stopped)
+	{
+		SetProcessingStage( EStatsProcessingStage::SPS_PostProcessStats );
+	}
+}
+
+bool FStatsReadFile::UpdateReadStageProgress()
+{
 	const double CurrentSeconds = FPlatformTime::Seconds();
 	if (CurrentSeconds > LastUpdateTime + NumSecondsBetweenUpdates)
@@ -833,48 +1012,73 @@
 		UE_LOG( LogStats, Verbose, TEXT( "%3i%%, current num frames %4i" ), PercentagePos, CombinedHistory.Num() );
 		LastUpdateTime = CurrentSeconds;
 	}
+
+	// Abandon support.
+	if (bShouldStopProcessing == true)
+	{
+		SetProcessingStage( EStatsProcessingStage::SPS_Stopped );
+		return false;
+	}
+	else
+	{
+		return true;
+	}
 }
 
 void FStatsReadFile::UpdateCombinedHistoryStats()
 {
 	// Dump frame stats
-	int64 LastFrameNum = 0;
 	for (const auto& It : CombinedHistory)
 	{
-		const int64 FrameNum = It.Key;
-		int64 FramePacketsSize = 0;
-		int64 FrameStatMessages = 0;
-		int64 FramePackets = It.Value.Packets.Num(); // Threads
+		const int32 FrameNum = It.Key;
+		int32 FramePacketsSize = 0;
+		int32 FrameStatMessages = 0;
+		int32 FramePackets = It.Value.Packets.Num(); // Threads
 		for (const auto& It2 : It.Value.Packets)
 		{
 			FramePacketsSize += It2->StatMessages.GetAllocatedSize();
 			FrameStatMessages += It2->StatMessages.Num();
 		}
 
-		UE_LOG( LogStats, Verbose, TEXT( "Frame: %4llu/%2lli Size: %5.1f MB / %10lli" ),
+		UE_LOG( LogStats, Verbose, TEXT( "Frame: %4i/%2i Size: %5.1f MB / %10i" ),
 				FrameNum,
 				FramePackets,
 				FramePacketsSize / 1024.0f / 1024.0f,
 				FrameStatMessages );
 
-		LastFrameNum = FMath::Max( LastFrameNum, FrameNum );
-
 		FileInfo.TotalStatMessagesNum += FrameStatMessages;
 		FileInfo.TotalPacketsSize += FramePacketsSize;
 	}
 
-	UE_LOG( LogStats, Warning, TEXT( "Total PacketSize: %6.1f MB, Max: %2f MB, PacketsNum: %lli, StatMessagesNum: %lli, Frames: %i" ),
+	UE_LOG( LogStats, Warning, TEXT( "Total PacketSize: %6.1f MB, Max: %2f MB, PacketsNum: %i, StatMessagesNum: %i, Frames: %i" ),
 			FileInfo.TotalPacketsSize / 1024.0f / 1024.0f,
 			FileInfo.MaximumPacketSize / 1024.0f / 1024.0f,
 			FileInfo.TotalPacketsNum,
 			FileInfo.TotalStatMessagesNum,
 			CombinedHistory.Num() );
-
-	TArray<int64> Frames;
-	CombinedHistory.GenerateKeyArray( Frames );
-	Frames.Sort();
-	// Verify that frames are sequential.
-	check( Frames[Frames.Num() - 1] == LastFrameNum );
 }
 
+bool FStatsReadFile::UpdateProcessStageProgress( const int32 CurrentStatMessageIndex, const int32 FrameIndex, const int32 PacketIndex )
+{	
+	const double CurrentSeconds = FPlatformTime::Seconds();
+	if (CurrentSeconds > LastUpdateTime + NumSecondsBetweenUpdates)
+	{
+		const int32 PercentagePos = int32( 100.0*CurrentStatMessageIndex / FileInfo.TotalStatMessagesNum );
+		StageProgress.Set( PercentagePos );
+		UE_LOG( LogStats, Verbose, TEXT( "Processing %3i%% (%10i/%10i) stat messages [Frame: %3i, Packet: %2i]" ), PercentagePos, CurrentStatMessageIndex, FileInfo.TotalStatMessagesNum, FrameIndex, PacketIndex );
+		LastUpdateTime = CurrentSeconds;	
+	}	
+
+	// Abandon support.
+	if (bShouldStopProcessing == true)
+	{
+		SetProcessingStage( EStatsProcessingStage::SPS_Stopped );
+		return false;
+	}
+	else
+	{
+		return true;
+	}
+}
+
 /*-----------------------------------------------------------------------------
diff --git a/Engine/Source/Runtime/Core/Public/Stats/StatsFile.h b/Engine/Source/Runtime/Core/Public/Stats/StatsFile.h
--- a/Engine/Source/Runtime/Core/Public/Stats/StatsFile.h
+++ b/Engine/Source/Runtime/Core/Public/Stats/StatsFile.h
@@ -755,29 +755,40 @@ struct FRawStatsFileInfo
 	{}
 
 	/** Size of all packets. */
-	int64 TotalPacketsSize;
+	int32 TotalPacketsSize;
 
 	/** Number of all stat messages. */
-	int64 TotalStatMessagesNum;
+	int32 TotalStatMessagesNum;
 
 	/** Maximum packet size. */
-	int64 MaximumPacketSize;
+	int32 MaximumPacketSize;
 
 	/** Number of all packets. */
-	int64 TotalPacketsNum;
+	int32 TotalPacketsNum;
 };
 
 /** Enumerates stats processing stages. */
 enum class EStatsProcessingStage : int32
 {
+	/** Started loading. */
 	SPS_Started = 0,
-	SPS_ReadAndCombinePackets,
-	SPS_ProcessStatMessages,
-	/** Memory profiler only. */
-	SPS_SortSequence,
-	/** Memory profiler only. */
-	SPS_ProcessAllocations,
+
+	/** Read and combine packets. */
+	SPS_ReadStats,
+
+	/** Pre process stats, prepare. */
+	SPS_PreProcessStats,
+
+	/** Process combine history. */
+	SPS_ProcessStats,
+
+	/** Post process stats, finalize. */
+	SPS_PostProcessStats,
+
+	/** Finished processing. */
 	SPS_Finished,
+
+	/** Stopped processing. */
 	SPS_Stopped,
 
 	/** Last stage, at this moment all data is read and processed or process has been stopped, we can now remove the instance. */
@@ -817,31 +828,60 @@ public:
 	void Abandon();
 };
 
+/*-----------------------------------------------------------------------------
+	Stats stack helpers
+-----------------------------------------------------------------------------*/
+
+/** Holds stats stack state, used to preserve continuity when the game frame has changed. For raw stats. */
+struct FStackState
+{
+	/** Default constructor. */
+	FStackState()
+		: bIsBrokenCallstack( false )
+	{}
+
+	/** Call stack. */
+	TArray<FName> Stack;
+
+	/** Current function name. */
+	FName Current;
+
+	/** Whether this callstack is marked as broken due to mismatched start and end scope cycles. */
+	bool bIsBrokenCallstack;
+};
+
+/*-----------------------------------------------------------------------------
+	FStatsReadFile
+-----------------------------------------------------------------------------*/
+
+template<typename T>
+struct FCreateStatsReader
+{
+	/** Creates a new reader for raw stats file. Will be nullptr for invalid files. */	
+	static T* ForRawStats( const TCHAR* Filename )
+	{
+		T* StatsReadFile = new T( Filename );
+		const bool bValid = StatsReadFile->PrepareLoading();
+		if (!bValid)
+		{
+			delete StatsReadFile;
+		}
+		return bValid ? StatsReadFile : nullptr;
+	}
+};
+
 /** Struct used to read from ue4stats/ue4statsraw files, initializes all metadata and starts a process of reading the file asynchronously. */
-// #YRX_Stats: 2015-07-08 Maybe add IStatsReadFile -> FStatsReadFile, FRawStatsReadFile
 // #YRX_Stats: 2015-07-09 Remove CORE_API
 struct CORE_API FStatsReadFile
 {
-protected:
+	friend class FAsyncRawStatsFile;
+
 	/** Number of seconds between updating the current stage. */
 	static const double NumSecondsBetweenUpdates;
 
-public:
-	/** Delegate that FStatsReadFile calls on the instance whenever we have a new frame. */
-	DECLARE_DELEGATE_OneParam( FOnNewStatPacket, TArray<FStatMessage>* );
-
-	/** Delegate that FStatsReadFile calls on the instance whenever we have a new raw stats packet. */
-	//DECLARE_DELEGATE_OneParam( FOnNewRawStatPacket, FStatPacket* );
-
-	/** Delegate that FStatsReadFile calls on the instance whenever we have a whole combined history read. */
-	DECLARE_DELEGATE_OneParam( FOnNewCombinedHistory, FStatsReadFile* );
-
 public:
 	/** Creates a new reader for regular stats file. Will be nullptr for invalid files. */
-	static FStatsReadFile* CreateReaderForRegularStats( const TCHAR* Filename, FOnNewStatPacket InNewStatPacket );
-
-	/** Creates a new reader for raw stats file. Will be nullptr for invalid files. */
-	static FStatsReadFile* CreateReaderForRawStats( const TCHAR* Filename, FOnNewCombinedHistory InNewCombinedHistory );
+	static FStatsReadFile* CreateReaderForRegularStats( const TCHAR* Filename );
 
 	/** Reads and processes the file on the current thread. This is a blocking operation. */
 	void ReadAndProcessSynchronously();
@@ -849,16 +889,80 @@ public:
 	/** Reads and processes the file using the async tasks on the pool thread. The read data is sent to the game thread using the task graph. This is a non-blocking operation. */
 	void ReadAndProcessAsynchronously();
 
+	/** Initialization constructor. */
+	FStatsReadFile( const TCHAR* InFilename, bool bInRawStatsFile );
+
+public:
+	/** Destructor. */
+	virtual ~FStatsReadFile();
+
+protected:
+	/**
+	 * Prepares file to be loaded, makes sanity checks, reads and initializes metadata
+	 * @return true, if the process was completed successfully
+	 */
+	bool PrepareLoading();
+
+	/** Reads stats from the file into combined history. */
 	void ReadStats();
 
+	/** Called before started processing combined history. */
+	virtual void PreProcessStats();
+
+	/** Processes combined history using the internal functionality and provided overloaded Process*Operation methods. */
 	void ProcessStats();
 
+	/** Called after finished processing combined history. */
+	virtual void PostProcessStats();
+
+	/** Processes special message for advancing the stats frame from the game thread. */
+	virtual void ProcessAdvanceFrameEventGameThreadOperation( const FStatMessage& Message, const FStackState& StackState )
+	{}
+
+	/** Processes special message for advancing the stats frame from the render thread. */
+	virtual void ProcessAdvanceFrameEventRenderThreadOperation( const FStatMessage& Message, const FStackState& StackState )
+	{}
+
+	/** ProcessesIndicates begin of the cycle scope. */
+	virtual void ProcessCycleScopeStartOperation( const FStatMessage& Message, const FStackState& StackState )
+	{}
+
+	/** Indicates end of the cycle scope. */
+	virtual void ProcessCycleScopeEndOperation( const FStatMessage& Message, const FStackState& StackState )
+	{}
+
+	/** Processes special message marker used determine that we encountered a special data in the stat file. */
+	virtual void ProcessSpecialMessageMarkerOperation( const FStatMessage& Message, const FStackState& StackState )
+	{}
+
+	// #YRX_Stats: 2015-07-13 Not implemented yet.
+// 	/** Processes set operation. */
+// 	virtual void ProcessSetOperation( const FStatMessage& Message, const FStackState& StackState )
+// 	{}
+// 
+// 	/** Processes clear operation. */
+// 	virtual void ProcessClearOperation( const FStatMessage& Message, const FStackState& StackState )
+// 	{}
+// 
+// 	/** Processes add operation. */
+// 	virtual void ProcessAddOperation( const FStatMessage& Message, const FStackState& StackState )
+// 	{}
+// 
+// 	/** Processes subtract operation. */
+// 	virtual void ProcessSubtractOperation( const FStatMessage& Message, const FStackState& StackState )
+// 	{}
+
+	/** Processes memory operation. @see EMemoryOperation. */
+	virtual void ProcessMemoryOperation( EMemoryOperation MemOp, uint64 Ptr, uint64 NewPtr, int64 Size, uint32 SequenceTag, const FStackState& StackState )
+	{}
+
 	/** Sets a new processing stage for this file. */
 	void SetProcessingStage( EStatsProcessingStage NewStage )
 	{
 		ProcessingStage.Set( int32( NewStage ) );
 	}
 
+public:
 	/**
 	 * @return current processing stage.
 	 */
@@ -878,21 +982,21 @@ public:
 		{
 			Result = TEXT( "SPS_Started" );
 		}
-		else if (Stage == EStatsProcessingStage::SPS_ReadAndCombinePackets)
+		else if (Stage == EStatsProcessingStage::SPS_ReadStats)
 		{
-			Result = TEXT( "SPS_ReadAndCombinePackets" );
+			Result = TEXT( "SPS_ReadStats" );
 		}
-		else if (Stage == EStatsProcessingStage::SPS_ProcessStatMessages)
+		else if (Stage == EStatsProcessingStage::SPS_PreProcessStats)
 		{
-			Result = TEXT( "SPS_ProcessStatMessages" );
+			Result = TEXT( "SPS_PreProcessStats" );
 		}
-		else if (Stage == EStatsProcessingStage::SPS_SortSequence)
+		else if (Stage == EStatsProcessingStage::SPS_ProcessStats)
 		{
-			Result = TEXT( "SPS_SortSequence" );
+			Result = TEXT( "SPS_ProcessStats" );
 		}
-		else if (Stage == EStatsProcessingStage::SPS_ProcessAllocations)
+		else if (Stage == EStatsProcessingStage::SPS_PostProcessStats)
 		{
-			Result = TEXT( "SPS_ProcessAllocations" );
+			Result = TEXT( "SPS_PostProcessStats" );
 		}
 		else if (Stage == EStatsProcessingStage::SPS_Finished)
 		{
@@ -932,29 +1036,25 @@ public:
 		bShouldStopProcessing = true;
 	}
 
-protected:
-	/** Initialization constructor. */
-	FStatsReadFile( const TCHAR* InFilename, bool bInRawStatsFile );
-
-public:
-	/** Destructor. */
-	~FStatsReadFile();
-
 protected:
 	/**
-	 * Prepares file to be loaded, makes sanity checks, reads and initializes metadata
-	 * @return true, if the process was completed successfully
+	 * Updates read stage progress periodically, does debug logging if enabled.
+	 *
+	 * @return true, if finished, false otherwise like stopped.
 	 */
-	bool PrepareLoading();
-
-	/** Updates stage progress periodically, does debug logging if enabled. */
-	void UpdateReadStageProgress();
+	bool UpdateReadStageProgress();
 
 	/** Dumps combined history stats. Only for raw stats. */
 	void UpdateCombinedHistoryStats();
 
-//protected:
-public:
+	/**
+	 * Updates process stage progress periodically, does debug logging if enabled.
+	 *
+	 * @return true, if finished, false otherwise like stopped.
+	 */
+	bool UpdateProcessStageProgress( const int32 CurrentStatMessageIndex, const int32 FrameIndex, const int32 PacketIndex );
+
+protected:
 	/** Current state of the stats. Mostly for metadata. */
 	FStatsThreadState State;
 
@@ -970,23 +1070,14 @@ public:
 	/** Async task. */
 	FAsyncTask<FAsyncRawStatsFile>* AsyncWork;
 
-	/** Delegate we fire every time we read one frame of data. The data must be freed after processing. */
-	FOnNewStatPacket NewStatPacket;
-
-	/** Delegate we fire every time we read one raw stats packet. The data must be freed after processing. */
-	//FOnNewRawStatPacket NewRawStatPacket;
-	
-	/**
-	 *	Delegate we fire after reading all raw stats packets and combining them into one history. 
-	 *	The combined history will be freed after processing.
-	 */
-	FOnNewCombinedHistory NewCombinedHistory;
-
 	/** Basic information about the stats file. */
 	FRawStatsFileInfo FileInfo;
 
 	/** Combined history for raw packets, indexed by a frame number. */
-	TMap<int64, FStatPacketArray> CombinedHistory;
+	TMap<int32, FStatPacketArray> CombinedHistory;
+
+	/** Frames computed from the combined history. */
+	TArray<int32> Frames;
 
 	/** All raw names that contains a path to an UObject. */
 	TSet<FName> UObjectRawNames;
