UnrealFrontend: fix Linux build.

diff --git a/Engine/Source/Programs/UnrealFrontend/Private/Linux/LinuxUnrealFrontendMain.cpp b/Engine/Source/Programs/UnrealFrontend/Private/Linux/LinuxUnrealFrontendMain.cpp
--- a/Engine/Source/Programs/UnrealFrontend/Private/Linux/LinuxUnrealFrontendMain.cpp
+++ b/Engine/Source/Programs/UnrealFrontend/Private/Linux/LinuxUnrealFrontendMain.cpp
@@ -2,7 +2,8 @@
 
 #include "UnrealFrontendPrivatePCH.h"
 #include "UnrealFrontendMain.h"
-
+#include "ExceptionHandling.h"
+#include "RequiredProgramMainCPPInclude.h"
 
 static FString GSavedCommandLine;
 
