This submission allows the end user to launch onto any platform using a profile created inside of the UFE.
Allows for easier scripting for testing and launching your project.

To run:
UnrealFrontend.exe -Run="LAUNCHPROFILE" -PROFILENAME="THE_NAME_OF_YOUR_CUSTOM_PROFILE"

**Changes and Additions**
Added LaunchFromProfileCommand.h/.cpp
--Checks if '-ProfileName=' is used and is correct.
--Loads the LaucherServicesModules and DeviceServicesModules
----Gets the Device Proxy Manager
----Gets the Profile Manager which is needed to find the named profile.
--A launcher is created that uses the profile and DeviceProxyManager to read the profile and trigger a build, cook, and launch of it.
--Logging from the launch will be shown in the command window.
--Keeps running until the launcher broadcasts that it has either completed or has been cancelled.

UnrealFrontendMain.cpp
Task Related:
--Added the launch profile command to the command if/else statement.
--Added helpful comments.
Improvements and General Fixes:
--Added a check to verify that a command is actually given with -Run.  Logs a warning if none exist.
--Fixed an issue where the -messaging command line check was returning incorrectly.  It will now only add '-messaging' if it actually does not exist in the command line already.
--Added a command line check where if you run with -run then -log will be appended to the command line if it isn't already.  This allows us to receive logging onscreen.

--------
Integrated using branch UE4-Orion-To-UE4 of change#2675868 by [REDACTED_PERSON] on 2015/09/01 14:50:30.

diff --git a/Engine/Source/Programs/UnrealFrontend/Private/Commands/LaunchFromProfileCommand.cpp b/Engine/Source/Programs/UnrealFrontend/Private/Commands/LaunchFromProfileCommand.cpp
new file mode 100644
--- /dev/null
+++ b/Engine/Source/Programs/UnrealFrontend/Private/Commands/LaunchFromProfileCommand.cpp
@@ -0,0 +1,83 @@
+// Copyright 1998-2015 Epic Games, Inc. All Rights Reserved.
+
+#include "UnrealFrontendPrivatePCH.h"
+#include "LaunchFromProfileCommand.h"
+#include "DesktopPlatformModule.h"
+#include "LauncherServices.h"
+#include "UserInterfaceCommand.h"
+
+DEFINE_LOG_CATEGORY_STATIC(LogUFECommands, All, All);
+
+void FLaunchFromProfileCommand::Run(const FString& Params)
+{
+	// Get the name of the profile from the command line.
+	FString ProfileName;
+	FParse::Value(FCommandLine::Get(), TEXT("-PROFILENAME="), ProfileName);
+	if (ProfileName.IsEmpty())
+	{
+		UE_LOG(LogUFECommands, Warning, TEXT("Profile name was found.  Please use '-PROFILENAME=' in your command line."));
+		return;
+	}
+	
+	// Loading the launcher services module to get the needed profile.
+	ILauncherServicesModule& LauncherServicesModule = FModuleManager::LoadModuleChecked<ILauncherServicesModule>(TEXT("LauncherServices"));
+	ILauncherProfileManagerRef ProfileManager = LauncherServicesModule.GetProfileManager();
+	ILauncherProfilePtr Profile = ProfileManager->FindProfile(ProfileName);
+
+	// Loading the Device Proxy Manager to get the needed Device Manager.
+	ITargetDeviceServicesModule& DeviceServiceModule = FModuleManager::LoadModuleChecked<ITargetDeviceServicesModule>(TEXT("TargetDeviceServices"));
+	ITargetDeviceProxyManagerRef DeviceProxyManager = DeviceServiceModule.GetDeviceProxyManager();
+
+	UE_LOG(LogUFECommands, Log, TEXT("Begin the process of launching a project using the provided profile."));
+	ILauncherRef LauncherRef = LauncherServicesModule.CreateLauncher();
+	ILauncherWorkerPtr LauncherWorkerPtr = LauncherRef->Launch(DeviceProxyManager, Profile.ToSharedRef());
+
+	// This will allow us to pipe the launcher messages into the command window.
+	LauncherWorkerPtr.Get()->OnOutputReceived().AddStatic(&FLaunchFromProfileCommand::MessageReceived);
+	// Allows us to exit this command once the launcher worker has completed or is canceled 
+	LauncherWorkerPtr.Get()->OnCompleted().AddRaw(this, &FLaunchFromProfileCommand::LaunchCompleted);
+	LauncherWorkerPtr.Get()->OnCanceled().AddRaw(this, &FLaunchFromProfileCommand::LaunchCanceled);
+
+	TArray<ILauncherTaskPtr> TaskList;
+	int32 NumOfTasks = LauncherWorkerPtr->GetTasks(TaskList);	
+	UE_LOG(LogUFECommands, Log, TEXT("There are '%i' tasks to be completed."), NumOfTasks);
+
+	// Holds the current element in the TaskList array.
+	int32 TaskIndex = 0;
+	// Holds the name of the current tasked.
+	FString TriggeredTask;
+
+	bTestRunning = true;
+	while (bTestRunning)
+	{
+		if (TaskIndex >= NumOfTasks) continue;
+		ILauncherTaskPtr CurrentTask = TaskList[TaskIndex];
+
+		// Log the current task, but only once per run.
+		if (CurrentTask->GetStatus() == ELauncherTaskStatus::Busy)
+		{
+			if (CurrentTask->GetDesc() == TriggeredTask) continue;
+
+			TriggeredTask = *CurrentTask->GetDesc();
+			UE_LOG(LogUFECommands, Log, TEXT("Current Task is %s"), *TriggeredTask);
+			TaskIndex++;
+		}
+	}
+}
+
+void FLaunchFromProfileCommand::MessageReceived(const FString& InMessage)
+{
+	GLog->Logf(ELogVerbosity::Log, TEXT("%s"), *InMessage);
+}
+
+void FLaunchFromProfileCommand::LaunchCompleted(bool Outcome, double ExecutionTime, int32 ReturnCode)
+{
+	UE_LOG(LogUFECommands, Log, TEXT("Profile launch command %s."), Outcome ? TEXT("is SUCCESSFUL") : TEXT("has FAILED"));
+	bTestRunning = false;
+}
+
+void FLaunchFromProfileCommand::LaunchCanceled(double ExecutionTime)
+{
+	UE_LOG(LogUFECommands, Log, TEXT("Profile launch command was canceled."));
+	bTestRunning = false;
+}
\ No newline at end of file
diff --git a/Engine/Source/Programs/UnrealFrontend/Private/Commands/LaunchFromProfileCommand.h b/Engine/Source/Programs/UnrealFrontend/Private/Commands/LaunchFromProfileCommand.h
new file mode 100644
--- /dev/null
+++ b/Engine/Source/Programs/UnrealFrontend/Private/Commands/LaunchFromProfileCommand.h
@@ -0,0 +1,27 @@
+// Copyright 1998-2015 Epic Games, Inc. All Rights Reserved.
+
+#pragma once
+
+class FLaunchFromProfileCommand
+{
+public:
+
+	/**
+	* Executes the command.
+	* Use '-PROFILENAME=' command line argument to specify the name of a profile to use for this command.
+	*/
+	void Run(const FString& Params);
+
+protected:
+
+	/**
+	* Logs out the messages from the launcher worker.  Used in a delegate.
+	*/
+	static void MessageReceived(const FString& InMessage);
+
+	void LaunchCompleted(bool Outcome, double ExecutionTime, int32 ReturnCode);
+	void LaunchCanceled(double ExecutionTime);
+
+private:
+	bool bTestRunning;
+};
diff --git a/Engine/Source/Programs/UnrealFrontend/Private/UnrealFrontendMain.cpp b/Engine/Source/Programs/UnrealFrontend/Private/UnrealFrontendMain.cpp
--- a/Engine/Source/Programs/UnrealFrontend/Private/UnrealFrontendMain.cpp
+++ b/Engine/Source/Programs/UnrealFrontend/Private/UnrealFrontendMain.cpp
@@ -10,6 +10,7 @@
 #include "StatsConvertCommand.h"
 #include "StatsDumpMemoryCommand.h"
 #include "UserInterfaceCommand.h"
+#include "LaunchFromProfileCommand.h"
 
 
 IMPLEMENT_APPLICATION(UnrealFrontend, "UnrealFrontend");
@@ -29,9 +30,14 @@ int32 UnrealFrontendMain( const TCHAR* CommandLine )
 	// process command line parameters
 	bool bRunCommand = FParse::Value(*NewCommandLine, TEXT("-RUN="), Command);
 	
+	if (Command.IsEmpty())
+	{
+		GLog->Logf(ELogVerbosity::Warning, TEXT("The command line argument '-RUN=' does not have a command name associated with it."));
+	}
+
 	if (bRunCommand)
 	{
-		// extract off any -PARAM= parameters so that they aren't accidentally parsed by engine init
+		// extract off any '-PARAM=' parameters so that they aren't accidentally parsed by engine init
 		FParse::Value(*NewCommandLine, TEXT("-PARAMS="), Params);
 
 		if (Params.Len() > 0)
@@ -42,18 +48,27 @@ int32 UnrealFrontendMain( const TCHAR* CommandLine )
 		}
 	}
 
-	if (!FParse::Param(*NewCommandLine, TEXT("-Messaging")))
+	// Add -Messaging if it was not given in the command line.
+	if (!FParse::Param(*NewCommandLine, TEXT("MESSAGING")))
 	{
 		NewCommandLine += TEXT(" -Messaging");
 	}
 
+	// Add '-Log' if the Frontend was run with '-RUN=' without '-LOG' so we can read any potential log output.
+	if (bRunCommand && !FParse::Param(*NewCommandLine, TEXT("LOG")))
+	{
+		NewCommandLine += TEXT(" -Log");
+	}
+
 	// initialize core
 	GEngineLoop.PreInit(*NewCommandLine);
 	FModuleManager::Get().StartProcessingNewlyLoadedObjects();
 
 	bool Succeeded = true;
 
-	// execute desired command
+	// Execute desired command
+	// To execute, run with '-RUN="COMMAND_NAME_FOUND_BELOW"'. 
+	// NOTE - Some commands may require extra command line parameters.
 	if (bRunCommand)
 	{
 		if (Command.Equals(TEXT("PACKAGE"), ESearchCase::IgnoreCase))
@@ -76,6 +91,12 @@ int32 UnrealFrontendMain( const TCHAR* CommandLine )
 		{
 			FStatsMemoryDumpCommand::Run();
 		}
+		// The 'LAUNCHPROFILE' command also needs '-PROFILENAME="MY_PROFILE_NAME"' as a command line parameter.
+		else if (Command.Equals(TEXT("LAUNCHPROFILE"), ESearchCase::IgnoreCase))
+		{
+			FLaunchFromProfileCommand* ProfileLaunch = new FLaunchFromProfileCommand;
+			ProfileLaunch->Run(Params);
+		}
 	}
 	else
 	{
