Use STAT_SecondsPerCycle to scale CVS stats instead of the host machine's default platform value.

diff --git a/Engine/Source/Programs/UnrealFrontend/Private/Commands/StatsConvertCommand.cpp b/Engine/Source/Programs/UnrealFrontend/Private/Commands/StatsConvertCommand.cpp
--- a/Engine/Source/Programs/UnrealFrontend/Private/Commands/StatsConvertCommand.cpp
+++ b/Engine/Source/Programs/UnrealFrontend/Private/Commands/StatsConvertCommand.cpp
@@ -208,11 +208,22 @@ void FStatsConvertCommand::CollectAndWriteStatsValues( FArchive& Writer )
 	// get the thread stats
 	TArray<FStatMessage> Stats;
 	ThreadState.GetInclusiveAggregateStackStats(ThreadState.CurrentGameFrame, Stats);
+
+	// The tick rate for different platforms will be different. We cannot use the Windows tick rate accurately here. We will pull it from the stats,
+	// but until we encounter one our best bet is to leave values as full ticks that can be analysed by hand.
+	double MillisecondsPerCycle = 1.0f;
+
 	for (int32 Index = 0; Index < Stats.Num(); ++Index)
 	{
 		FStatMessage const& Meta = Stats[Index];
 		//UE_LOG(LogTemp, Display, TEXT("Stat: %s"), *Meta.NameAndInfo.GetShortName().ToString());
 
+		if (Meta.NameAndInfo.GetShortName() == TEXT("STAT_SecondsPerCycle"))
+		{
+			// SecondsPerCycle may vary over time, so we update it here
+			MillisecondsPerCycle = Meta.GetValue_double() * 1000.0f;
+		}
+
 		for (int32 Jndex = 0; Jndex < StatList.Num(); ++Jndex)
 		{
 			if (Meta.NameAndInfo.GetShortName() == StatList[Jndex])
@@ -220,11 +231,11 @@ void FStatsConvertCommand::CollectAndWriteStatsValues( FArchive& Writer )
 				double StatValue = 0.0f;
 				if (Meta.NameAndInfo.GetFlag(EStatMetaFlags::IsPackedCCAndDuration))
 				{
-					StatValue = FPlatformTime::ToMilliseconds( FromPackedCallCountDuration_Duration(Meta.GetValue_int64()) );
+					StatValue = MillisecondsPerCycle * FromPackedCallCountDuration_Duration(Meta.GetValue_int64());
 				}
 				else if (Meta.NameAndInfo.GetFlag(EStatMetaFlags::IsCycle))
 				{
-					StatValue = FPlatformTime::ToMilliseconds( Meta.GetValue_int64() );
+					StatValue = MillisecondsPerCycle * Meta.GetValue_int64();
 				}
 				else
 				{
