Fix broken icons in UnrealFrontend.

#rnx

diff --git a/Engine/Source/Programs/UnrealFrontend/UnrealFrontend.Target.cs b/Engine/Source/Programs/UnrealFrontend/UnrealFrontend.Target.cs
--- a/Engine/Source/Programs/UnrealFrontend/UnrealFrontend.Target.cs
+++ b/Engine/Source/Programs/UnrealFrontend/UnrealFrontend.Target.cs
@@ -18,6 +18,9 @@ public class UnrealFrontendTarget : TargetRules
 		bCompileWithStatsWithoutEngine = true;
 		bCompileWithPluginSupport = true;
 
+		// For UI functionality
+		bBuildDeveloperTools = true;
+
 		bHasExports = false;
 	}
 }
