Call AppExit on from the AppExitCallback in UnixCommonMain

diff --git a/Engine/Source/Programs/UnrealFrontend/Private/Linux/LinuxUnrealFrontendMain.cpp b/Engine/Source/Programs/UnrealFrontend/Private/Linux/LinuxUnrealFrontendMain.cpp
--- a/Engine/Source/Programs/UnrealFrontend/Private/Linux/LinuxUnrealFrontendMain.cpp
+++ b/Engine/Source/Programs/UnrealFrontend/Private/Linux/LinuxUnrealFrontendMain.cpp
@@ -2,9 +2,10 @@
 
 #include "CoreMinimal.h"
 #include "UnrealFrontendMain.h"
+#include "LaunchEngineLoop.h"
 #include "UnixCommonStartup.h"
 
 int main(int argc, char *argv[])
 {
-	return CommonUnixMain(argc, argv, &UnrealFrontendMain);
+	return CommonUnixMain(argc, argv, &UnrealFrontendMain, [] { FEngineLoop::AppExit(); } );
 }
