Enable high DPI awareness in UnrealFrontend.

diff --git a/Engine/Source/Programs/UnrealFrontend/Private/Commands/UserInterfaceCommand.cpp b/Engine/Source/Programs/UnrealFrontend/Private/Commands/UserInterfaceCommand.cpp
--- a/Engine/Source/Programs/UnrealFrontend/Private/Commands/UserInterfaceCommand.cpp
+++ b/Engine/Source/Programs/UnrealFrontend/Private/Commands/UserInterfaceCommand.cpp
@@ -114,6 +114,8 @@ void FUserInterfaceCommand::Run(  )
 void FUserInterfaceCommand::InitializeSlateApplication( const FString& LayoutIni )
 {
 	FSlateApplication::InitializeAsStandaloneApplication(GetStandardStandaloneRenderer());
+	FSlateApplication::InitHighDPI(true);
+
 	FGlobalTabmanager::Get()->SetApplicationTitle(NSLOCTEXT("UnrealFrontend", "AppTitle", "Unreal Frontend"));
 
 	// load widget reflector
@@ -159,7 +161,6 @@ void FUserInterfaceCommand::InitializeSlateApplication( const FString& LayoutIni
 	FGlobalTabmanager::Get()->RestoreFrom(UserInterfaceCommand::ApplicationLayout.ToSharedRef(), TSharedPtr<SWindow>());
 }
 
-
 void FUserInterfaceCommand::ShutdownSlateApplication( const FString& LayoutIni )
 {
 	check(UserInterfaceCommand::ApplicationLayout.IsValid());
