Updates a UE4 named class with a new and improved 'Unreal' prefix

#rnx

diff --git a/Engine/Source/Programs/UnrealFrontend/Private/Mac/MacUnrealFrontendMain.cpp b/Engine/Source/Programs/UnrealFrontend/Private/Mac/MacUnrealFrontendMain.cpp
--- a/Engine/Source/Programs/UnrealFrontend/Private/Mac/MacUnrealFrontendMain.cpp
+++ b/Engine/Source/Programs/UnrealFrontend/Private/Mac/MacUnrealFrontendMain.cpp
@@ -11,14 +11,14 @@
 static FString GSavedCommandLine;
 
 
-@interface UE4AppDelegate : NSObject<NSApplicationDelegate, NSFileManagerDelegate>
+@interface UnrealAppDelegate : NSObject<NSApplicationDelegate, NSFileManagerDelegate>
 {
 }
 
 @end
 
 
-@implementation UE4AppDelegate
+@implementation UnrealAppDelegate
 
 //handler for the quit apple event used by the Dock menu
 - (void)handleQuitEvent:(NSAppleEventDescriptor*)Event withReplyEvent:(NSAppleEventDescriptor*)ReplyEvent
@@ -31,7 +31,7 @@ static FString GSavedCommandLine;
 	RequestEngineExit(TEXT("requestQuit"));
 }
 
-- (void) runGameThread:(id)Arg
+- (void)runGameThread:(id)Arg
 {
 	FPlatformMisc::SetGracefulTerminationHandler();
 	FPlatformMisc::SetCrashHandler(nullptr);
@@ -113,7 +113,7 @@ int main( int argc, char *argv[] )
 
 	SCOPED_AUTORELEASE_POOL;
 	[NSApplication sharedApplication];
-	[NSApp setDelegate:[UE4AppDelegate new]];
+	[NSApp setDelegate:[UnrealAppDelegate new]];
 	[NSApp run];
 	return 0;
 }
