fixed UnrealFrontend shutdown crash

diff --git a/Engine/Source/Programs/UnrealFrontend/Private/Windows/WindowsUnrealFrontendMain.cpp b/Engine/Source/Programs/UnrealFrontend/Private/Windows/WindowsUnrealFrontendMain.cpp
--- a/Engine/Source/Programs/UnrealFrontend/Private/Windows/WindowsUnrealFrontendMain.cpp
+++ b/Engine/Source/Programs/UnrealFrontend/Private/Windows/WindowsUnrealFrontendMain.cpp
@@ -19,6 +19,8 @@
  */
 int32 WINAPI WinMain( HINSTANCE hInInstance, HINSTANCE hPrevInstance, char* lpCmdLine, int32 nShowCmd )
 {
+	FTaskTagScope Scope(ETaskTag::EGameThread);
+
 	hInstance = hInInstance;
 
 	const TCHAR* CmdLine = ::GetCommandLineW();
