Fix the Convert command from UnrealFrontend
#prefligth 61a4c44647752856c9fee026

diff --git a/Engine/Source/Programs/UnrealFrontend/Private/Commands/StatsConvertCommand.cpp b/Engine/Source/Programs/UnrealFrontend/Private/Commands/StatsConvertCommand.cpp
--- a/Engine/Source/Programs/UnrealFrontend/Private/Commands/StatsConvertCommand.cpp
+++ b/Engine/Source/Programs/UnrealFrontend/Private/Commands/StatsConvertCommand.cpp
@@ -26,7 +26,7 @@ protected:
 	virtual void ReadStatsFrame( const TArray<FStatMessage>& CondensedMessages, const int64 Frame ) override;
 
 	/** Writes a formatted string to the CSV file. */
-	void WriteString( const ANSICHAR* Format, ... );
+	void WriteString( const TCHAR* Format, ... );
 public:
 
 	/** Sets a writer used to serialize the data for the CSV file. */
@@ -47,7 +47,7 @@ public:
 		}
 
 		// Output the csv header.
-		WriteString( "Frame,Name,Value\r\n" );
+		WriteString( TEXT("Frame,Name,Value\r\n") );
 	}
 
 protected:
@@ -106,23 +106,23 @@ void FCSVStatsProfiler::ReadStatsFrame( const TArray<FStatMessage>& CondensedMes
 				}
 
 				// write out to the csv file
-				WriteString( "%d,%S,%S\r\n", Frame, *StatShortNames[Jndex].ToString(), *StatValue );
+				WriteString( TEXT("%d,%s,%s\r\n"), Frame, *StatShortNames[Jndex].ToString(), *StatValue );
 			}
 		}
 	}
 }
 
-void FCSVStatsProfiler::WriteString( const ANSICHAR* Format, ... )
+void FCSVStatsProfiler::WriteString( const TCHAR* Format, ... )
 {
-	ANSICHAR Array[1024];
+	TCHAR Array[1024];
 	va_list ArgPtr;
 	va_start( ArgPtr, Format );
 	// Build the string.
-	int32 Result = FCStringAnsi::GetVarArgs( Array, UE_ARRAY_COUNT( Array ), Format, ArgPtr );
+	int32 Result = FCString::GetVarArgs( Array, UE_ARRAY_COUNT( Array ), Format, ArgPtr );
 	va_end(ArgPtr);
 
 	// Now write that to the file.
-	CSVWriter->Serialize( (void*)Array, Result );
+	CSVWriter->Serialize( (void*)Array, Result * sizeof(TCHAR) );
 }
 
 void FStatsConvertCommand::InternalRun()
