Fix UFE crashing at boot due to missing editor styles
- Still asserts but can skip past it

diff --git a/Engine/Source/Programs/UnrealFrontend/Private/Commands/UserInterfaceCommand.cpp b/Engine/Source/Programs/UnrealFrontend/Private/Commands/UserInterfaceCommand.cpp
--- a/Engine/Source/Programs/UnrealFrontend/Private/Commands/UserInterfaceCommand.cpp
+++ b/Engine/Source/Programs/UnrealFrontend/Private/Commands/UserInterfaceCommand.cpp
@@ -39,6 +39,8 @@ void FUserInterfaceCommand::Run(  )
 	FConfigCacheIni::InitializeConfigSystem();
 	GetTargetPlatformManager();
 
+	FCoreStyle::ResetToDefault();
+
 	// load required modules
 	FModuleManager::Get().LoadModuleChecked("Messaging");
 	FModuleManager::Get().LoadModuleChecked("OutputLog");
