remove duplicate task tag.

 - this makes Windows consistent with Linux & Mac.
 - the tag already exists in UnrealFrontendMain.

#rnx

diff --git a/Engine/Source/Programs/UnrealFrontend/Private/Windows/WindowsUnrealFrontendMain.cpp b/Engine/Source/Programs/UnrealFrontend/Private/Windows/WindowsUnrealFrontendMain.cpp
--- a/Engine/Source/Programs/UnrealFrontend/Private/Windows/WindowsUnrealFrontendMain.cpp
+++ b/Engine/Source/Programs/UnrealFrontend/Private/Windows/WindowsUnrealFrontendMain.cpp
@@ -19,8 +19,6 @@
  */
 int32 WINAPI WinMain( HINSTANCE hInInstance, HINSTANCE hPrevInstance, char* lpCmdLine, int32 nShowCmd )
 {
-	FTaskTagScope Scope(ETaskTag::EGameThread);
-
 	hInstance = hInInstance;
 
 	const TCHAR* CmdLine = ::GetCommandLineW();
