Add MessagingDebugger to UFE when the -DebugTools parameter is used.

#rnx

diff --git a/Engine/Plugins/Messaging/MessagingDebugger/MessagingDebugger.uplugin b/Engine/Plugins/Messaging/MessagingDebugger/MessagingDebugger.uplugin
--- a/Engine/Plugins/Messaging/MessagingDebugger/MessagingDebugger.uplugin
+++ b/Engine/Plugins/Messaging/MessagingDebugger/MessagingDebugger.uplugin
@@ -10,7 +10,10 @@
 	"Category" : "Messaging",
 	"EnabledByDefault" : false,
 	"IsBetaVersion": true,
-
+	"SupportedPrograms" : 
+	[
+		"UnrealFrontend"
+	],
 	"Modules" :
 	[
 		{
diff --git a/Engine/Programs/UnrealFrontend/Config/DefaultEngine.ini b/Engine/Programs/UnrealFrontend/Config/DefaultEngine.ini
--- a/Engine/Programs/UnrealFrontend/Config/DefaultEngine.ini
+++ b/Engine/Programs/UnrealFrontend/Config/DefaultEngine.ini
@@ -2,7 +2,7 @@
 GameNameShort=UFE
 
 [Plugins]
-+EnabledPlugins="MessagingDebugger"
++ProgramEnabledPlugins="MessagingDebugger"
 -EnabledPlugins="ExampleDeviceProfileSelector"
 +ProgramEnabledPlugins="UdpMessaging"
 +ProgramEnabledPlugins="TcpMessaging"
diff --git a/Engine/Source/Programs/UnrealFrontend/Private/Commands/UserInterfaceCommand.cpp b/Engine/Source/Programs/UnrealFrontend/Private/Commands/UserInterfaceCommand.cpp
--- a/Engine/Source/Programs/UnrealFrontend/Private/Commands/UserInterfaceCommand.cpp
+++ b/Engine/Source/Programs/UnrealFrontend/Private/Commands/UserInterfaceCommand.cpp
@@ -117,7 +117,7 @@ void FUserInterfaceCommand::InitializeSlateApplication( const FString& LayoutIni
 
 	FGlobalTabmanager::Get()->SetApplicationTitle(NSLOCTEXT("UnrealFrontend", "AppTitle", "Unreal Frontend"));
 
-	// load widget reflector
+	// load debug tools
 	const bool bAllowDebugTools = FParse::Param(FCommandLine::Get(), TEXT("DebugTools"));
 
 	if (bAllowDebugTools)
@@ -127,6 +127,8 @@ void FUserInterfaceCommand::InitializeSlateApplication( const FString& LayoutIni
 		{
 			SlateReflectorModule->RegisterTabSpawner(UserInterfaceCommand::DeveloperTools);
 		}
+
+		FModuleManager::Get().LoadModule("MessagingDebugger");
 	}
 
 	const float DPIScaleFactor = FPlatformApplicationMisc::GetDPIScaleFactorAtPoint(10.0f, 10.0f);
@@ -141,18 +143,20 @@ void FUserInterfaceCommand::InitializeSlateApplication( const FString& LayoutIni
 					FTabManager::NewStack()
 						->AddTab(FName("DeviceOutputLog"), ETabState::OpenedTab)
 						->AddTab(FName("DeviceManager"), ETabState::OpenedTab)
-						->AddTab(FName("MessagingDebugger"), ETabState::ClosedTab)
 						->AddTab(FName("SessionFrontend"), ETabState::OpenedTab)
 						->AddTab(FName("ProjectLauncher"), ETabState::OpenedTab)
 				)
 		)
 		->AddArea
 		(
-			FTabManager::NewArea(600.0f * DPIScaleFactor, 600.0f * DPIScaleFactor)
+			FTabManager::NewArea(1024.0f * DPIScaleFactor, 600.0f * DPIScaleFactor)
 				->SetWindow(FVector2D(10.0f * DPIScaleFactor, 10.0f * DPIScaleFactor), false)
 				->Split
 				(
-					FTabManager::NewStack()->AddTab("WidgetReflector", bAllowDebugTools ? ETabState::OpenedTab : ETabState::ClosedTab)
+					FTabManager::NewStack()
+						->AddTab(FName("WidgetReflector"), bAllowDebugTools ? ETabState::OpenedTab : ETabState::ClosedTab)
+						->AddTab(FName("MessagingDebugger"), bAllowDebugTools ? ETabState::OpenedTab : ETabState::ClosedTab)
+
 				)
 		);
 
diff --git a/Engine/Source/Programs/UnrealFrontend/UnrealFrontend.Build.cs b/Engine/Source/Programs/UnrealFrontend/UnrealFrontend.Build.cs
--- a/Engine/Source/Programs/UnrealFrontend/UnrealFrontend.Build.cs
+++ b/Engine/Source/Programs/UnrealFrontend/UnrealFrontend.Build.cs
@@ -81,6 +81,7 @@ public class UnrealFrontend : ModuleRules
 				"Sockets",
 				"UdpMessaging",
                 "TcpMessaging",
+				"MessagingDebugger",
 			}
 		);
 
