Fix V3 include path

#rnx

diff --git a/Engine/Source/Programs/UnrealFrontend/Private/Commands/UserInterfaceCommand.cpp b/Engine/Source/Programs/UnrealFrontend/Private/Commands/UserInterfaceCommand.cpp
--- a/Engine/Source/Programs/UnrealFrontend/Private/Commands/UserInterfaceCommand.cpp
+++ b/Engine/Source/Programs/UnrealFrontend/Private/Commands/UserInterfaceCommand.cpp
@@ -15,7 +15,7 @@
 #include "Framework/Application/SlateApplication.h"
 #include "HAL/PlatformProcess.h"
 #include "HAL/PlatformApplicationMisc.h"
-#include "OutputLog/Public/OutputLogModule.h"
+#include "OutputLogModule.h"
 
 #define IDEAL_FRAMERATE 60;
 
