Fix auto update contents in Screenshot Comparison tab for UFE

diff --git a/Engine/Source/Programs/UnrealFrontend/Private/Commands/UserInterfaceCommand.cpp b/Engine/Source/Programs/UnrealFrontend/Private/Commands/UserInterfaceCommand.cpp
--- a/Engine/Source/Programs/UnrealFrontend/Private/Commands/UserInterfaceCommand.cpp
+++ b/Engine/Source/Programs/UnrealFrontend/Private/Commands/UserInterfaceCommand.cpp
@@ -10,12 +10,15 @@
 #include "Containers/Ticker.h"
 #include "Misc/CommandLine.h"
 #include "Misc/ConfigCacheIni.h"
+#include "Misc/App.h"
 #include "Framework/Docking/TabManager.h"
 #include "Framework/Docking/LayoutService.h"
 #include "Framework/Application/SlateApplication.h"
 #include "HAL/PlatformProcess.h"
 #include "HAL/PlatformApplicationMisc.h"
 #include "OutputLogModule.h"
+#include "IDirectoryWatcher.h"
+#include "DirectoryWatcherModule.h"
 
 #define IDEAL_FRAMERATE 60;
 
@@ -66,6 +69,10 @@ void FUserInterfaceCommand::Run(  )
 	// Load the source code access module
 	ISourceCodeAccessModule& SourceCodeAccessModule = FModuleManager::LoadModuleChecked<ISourceCodeAccessModule>( FName( "SourceCodeAccess" ) );
 
+	// Initialize directory watcher module
+	// Load the directory watcher module
+	FDirectoryWatcherModule& DirectoryWatcherModule = FModuleManager::LoadModuleChecked<FDirectoryWatcherModule>(FName("DirectoryWatcher"));
+
 	// Manually load in the source code access plugins, as standalone programs don't currently support plugins.
 #if PLATFORM_MAC
 	IModuleInterface& XCodeSourceCodeAccessModule = FModuleManager::LoadModuleChecked<IModuleInterface>( FName( "XCodeSourceCodeAccess" ) );
@@ -87,6 +94,12 @@ void FUserInterfaceCommand::Run(  )
 
 		FTaskGraphInterface::Get().ProcessThreadUntilIdle(ENamedThreads::GameThread);
 
+		//We have to force tick here to be able to update Screen Comparison tab in the Slate applications without Engine loop
+		if (!FApp::IsProjectNameEmpty())
+		{
+			DirectoryWatcherModule.Get()->Tick(FApp::GetDeltaTime());
+		}
+
 		FSlateApplication::Get().PumpMessages();
 		FSlateApplication::Get().Tick();
 		FTSTicker::GetCoreTicker().Tick(DeltaTime);
diff --git a/Engine/Source/Programs/UnrealFrontend/UnrealFrontend.Build.cs b/Engine/Source/Programs/UnrealFrontend/UnrealFrontend.Build.cs
--- a/Engine/Source/Programs/UnrealFrontend/UnrealFrontend.Build.cs
+++ b/Engine/Source/Programs/UnrealFrontend/UnrealFrontend.Build.cs
@@ -75,6 +75,7 @@ public class UnrealFrontend : ModuleRules
 				"Sockets",
 				"UdpMessaging",
                 "TcpMessaging",
+				"DirectoryWatcher"
 			}
 		);
 
