UnrealFrontend: Removed the deprecated MEMORYDUMP run command.

diff --git a/Engine/Source/Programs/UnrealFrontend/Private/Commands/StatsDumpMemoryCommand.cpp b/Engine/Source/Programs/UnrealFrontend/Private/Commands/StatsDumpMemoryCommand.cpp
deleted file mode 100644
diff --git a/Engine/Source/Programs/UnrealFrontend/Private/Commands/StatsDumpMemoryCommand.h b/Engine/Source/Programs/UnrealFrontend/Private/Commands/StatsDumpMemoryCommand.h
deleted file mode 100644
diff --git a/Engine/Source/Programs/UnrealFrontend/Private/UnrealFrontendMain.cpp b/Engine/Source/Programs/UnrealFrontend/Private/UnrealFrontendMain.cpp
--- a/Engine/Source/Programs/UnrealFrontend/Private/UnrealFrontendMain.cpp
+++ b/Engine/Source/Programs/UnrealFrontend/Private/UnrealFrontendMain.cpp
@@ -4,7 +4,6 @@
 #include "RequiredProgramMainCPPInclude.h"
 
 #include "StatsConvertCommand.h"
-#include "StatsDumpMemoryCommand.h"
 #include "UserInterfaceCommand.h"
 #include "LaunchFromProfileCommand.h"
 
@@ -76,10 +75,6 @@ int32 UnrealFrontendMain( const TCHAR* CommandLine )
 		{
 			FStatsConvertCommand::Run();
 		}
-		else if( Command.Equals( TEXT("MEMORYDUMP"), ESearchCase::IgnoreCase ) )
-		{
-			FStatsMemoryDumpCommand::Run();
-		}
 		// The 'LAUNCHPROFILE' command also needs '-PROFILENAME="MY_PROFILE_NAME"' as a command line parameter.
 		else if (Command.Equals(TEXT("LAUNCHPROFILE"), ESearchCase::IgnoreCase))
 		{
