Fix additional UDP messaging crashes on Exit (in Unreal Frontend)
+ trigger postdefault plugins load in UnrealFrontend to get UDP Messaging to finalize loading

#rnx

diff --git a/Engine/Plugins/Messaging/UdpMessaging/Source/UdpMessaging/Private/Transport/UdpMessageProcessor.cpp b/Engine/Plugins/Messaging/UdpMessaging/Source/UdpMessaging/Private/Transport/UdpMessageProcessor.cpp
--- a/Engine/Plugins/Messaging/UdpMessaging/Source/UdpMessaging/Private/Transport/UdpMessageProcessor.cpp
+++ b/Engine/Plugins/Messaging/UdpMessaging/Source/UdpMessaging/Private/Transport/UdpMessageProcessor.cpp
@@ -384,7 +384,10 @@ void FUdpMessageProcessor::WaitAsyncTaskCompletion()
 	Stop();
 
 	// Wait for the processor thread, so we can access KnownNodes safely
-	Thread->WaitForCompletion();
+	if (Thread)
+	{
+		Thread->WaitForCompletion();
+	}
 
 	// Check if processor has in-flight serialization task(s).
 	auto HasIncompleteSerializationTasks = [this]()
diff --git a/Engine/Source/Programs/UnrealFrontend/Private/Commands/UserInterfaceCommand.cpp b/Engine/Source/Programs/UnrealFrontend/Private/Commands/UserInterfaceCommand.cpp
--- a/Engine/Source/Programs/UnrealFrontend/Private/Commands/UserInterfaceCommand.cpp
+++ b/Engine/Source/Programs/UnrealFrontend/Private/Commands/UserInterfaceCommand.cpp
@@ -75,6 +75,10 @@ void FUserInterfaceCommand::Run(  )
 	// Load the directory watcher module
 	FDirectoryWatcherModule& DirectoryWatcherModule = FModuleManager::LoadModuleChecked<FDirectoryWatcherModule>(FName("DirectoryWatcher"));
 
+	// trigger loading of post default plug-ins
+	// (for UDP Messaging)
+	IPluginManager::Get().LoadModulesForEnabledPlugins(ELoadingPhase::PostDefault);
+
 	// Manually load in the source code access plugins, as standalone programs don't currently support plugins.
 #if PLATFORM_MAC
 	IModuleInterface& XCodeSourceCodeAccessModule = FModuleManager::LoadModuleChecked<IModuleInterface>( FName( "XCodeSourceCodeAccess" ) );
