UnrealFrontend: Stats are required in Shipping builds

#rnx

diff --git a/Engine/Source/Programs/UnrealFrontend/UnrealFrontend.Target.cs b/Engine/Source/Programs/UnrealFrontend/UnrealFrontend.Target.cs
--- a/Engine/Source/Programs/UnrealFrontend/UnrealFrontend.Target.cs
+++ b/Engine/Source/Programs/UnrealFrontend/UnrealFrontend.Target.cs
@@ -9,8 +9,12 @@ public class UnrealFrontendTarget : TargetRules
 	{
 		Type = TargetType.Program;
 		LinkType = TargetLinkType.Modular;
+		IncludeOrderVersion = EngineIncludeOrderVersion.Latest;
 		AdditionalPlugins.Add("UdpMessaging");
 		LaunchModuleName = "UnrealFrontend";
+		
+		// Stats are required even in Shipping
+		GlobalDefinitions.Add("FORCE_USE_STATS=1");
 
 		bCompileAgainstEngine = false;
 		bCompileAgainstCoreUObject = true;
