Fix Trace Control channel toggle in Unreal Frontend

diff --git a/Engine/Source/Programs/UnrealFrontend/Private/Commands/UserInterfaceCommand.cpp b/Engine/Source/Programs/UnrealFrontend/Private/Commands/UserInterfaceCommand.cpp
--- a/Engine/Source/Programs/UnrealFrontend/Private/Commands/UserInterfaceCommand.cpp
+++ b/Engine/Source/Programs/UnrealFrontend/Private/Commands/UserInterfaceCommand.cpp
@@ -11,6 +11,7 @@
 #include "Misc/CommandLine.h"
 #include "Misc/ConfigCacheIni.h"
 #include "Misc/App.h"
+#include "Misc/CoreDelegates.h"
 #include "Framework/Docking/TabManager.h"
 #include "Framework/Docking/LayoutService.h"
 #include "Framework/Application/SlateApplication.h"
@@ -119,6 +120,7 @@ void FUserInterfaceCommand::Run(  )
 		LastTime = CurrentTime;
 
 		FStats::AdvanceFrame( false );
+		FCoreDelegates::OnEndFrame.Broadcast();
 
 		GLog->FlushThreadedLogs();
 	}
