Add CustomLaunchUI to UnrealFrontend & enable it for Mac and Linux too.

#rnx

diff --git a/Engine/Plugins/Developer/CustomLaunchUI/CustomLaunchUI.uplugin b/Engine/Plugins/Developer/CustomLaunchUI/CustomLaunchUI.uplugin
--- a/Engine/Plugins/Developer/CustomLaunchUI/CustomLaunchUI.uplugin
+++ b/Engine/Plugins/Developer/CustomLaunchUI/CustomLaunchUI.uplugin
@@ -15,13 +15,23 @@
 	"CanContainContent": false,
 	"IsBetaVersion": true,
 	"Installed": false,
-	"SupportedTargetPlatforms": [ "Win64" ],
+	"SupportedTargetPlatforms": [
+		"Win64",
+		"Linux",
+		"Mac"
+	],
+	"SupportedPrograms": [ "UnrealFrontend" ],
 	"Modules": [
 		{
 			"Name": "CustomLaunchUI",
 			"Type": "Editor",
 			"LoadingPhase": "Default",
-			"PlatformAllowList": [ "Win64" ]
+			"PlatformAllowList": [
+				"Win64",
+				"Linux",
+				"Mac"
+			],
+			"ProgramAllowList": [ "UnrealFrontend" ]
 		}
 	]
 }
\ No newline at end of file
diff --git a/Engine/Source/Programs/UnrealFrontend/Private/Commands/UserInterfaceCommand.cpp b/Engine/Source/Programs/UnrealFrontend/Private/Commands/UserInterfaceCommand.cpp
--- a/Engine/Source/Programs/UnrealFrontend/Private/Commands/UserInterfaceCommand.cpp
+++ b/Engine/Source/Programs/UnrealFrontend/Private/Commands/UserInterfaceCommand.cpp
@@ -66,6 +66,7 @@ void FUserInterfaceCommand::Run(  )
 	FModuleManager::Get().LoadModule("ProjectLauncher");
 	FModuleManager::Get().LoadModule("SessionFrontend");
 	FModuleManager::Get().LoadModule("SettingsEditor");
+	FModuleManager::Get().LoadModule("CustomLaunchUI");
 
 	InitializeSlateApplication(UnrealFrontendLayoutIni);
 
@@ -154,7 +155,7 @@ void FUserInterfaceCommand::InitializeSlateApplication( const FString& LayoutIni
 	const float DPIScaleFactor = FPlatformApplicationMisc::GetDPIScaleFactorAtPoint(10.0f, 10.0f);
 
 	// restore application layout
-	TSharedRef<FTabManager::FLayout> NewLayout = FTabManager::NewLayout("SessionFrontendLayout_v1.2")
+	TSharedRef<FTabManager::FLayout> NewLayout = FTabManager::NewLayout("SessionFrontendLayout_v1.3")
 		->AddArea
 		(
 			FTabManager::NewArea(1280.f * DPIScaleFactor, 720.0f * DPIScaleFactor)
@@ -166,6 +167,7 @@ void FUserInterfaceCommand::InitializeSlateApplication( const FString& LayoutIni
 						->AddTab(FName("MessagingDebugger"), ETabState::ClosedTab)
 						->AddTab(FName("SessionFrontend"), ETabState::OpenedTab)
 						->AddTab(FName("ProjectLauncher"), ETabState::OpenedTab)
+						->AddTab(FName("CustomLaunchUI"), ETabState::OpenedTab)
 				)
 		)
 		->AddArea
diff --git a/Engine/Source/Programs/UnrealFrontend/UnrealFrontend.Target.cs b/Engine/Source/Programs/UnrealFrontend/UnrealFrontend.Target.cs
--- a/Engine/Source/Programs/UnrealFrontend/UnrealFrontend.Target.cs
+++ b/Engine/Source/Programs/UnrealFrontend/UnrealFrontend.Target.cs
@@ -1,7 +1,6 @@
 // Copyright Epic Games, Inc. All Rights Reserved.
 
 using UnrealBuildTool;
-using System.Collections.Generic;
 
 public class UnrealFrontendTarget : TargetRules
 {
@@ -29,5 +28,8 @@ public class UnrealFrontendTarget : TargetRules
 
 		// Old Profiler (SessionFrontend/Profiler) is deprecated since UE 5.0. Use Trace/UnrealInsights instead.
 		//GlobalDefinitions.Add("UE_DEPRECATED_PROFILER_ENABLED=1");
+
+		// new custom launch UI
+		AdditionalPlugins.Add("CustomLaunchUI");
 	}
 }
