fix for UnrealFrontend not loading some plugins

#rnx

diff --git a/Engine/Source/Programs/UnrealFrontend/Private/Commands/UserInterfaceCommand.cpp b/Engine/Source/Programs/UnrealFrontend/Private/Commands/UserInterfaceCommand.cpp
--- a/Engine/Source/Programs/UnrealFrontend/Private/Commands/UserInterfaceCommand.cpp
+++ b/Engine/Source/Programs/UnrealFrontend/Private/Commands/UserInterfaceCommand.cpp
@@ -66,7 +66,7 @@ void FUserInterfaceCommand::Run(  )
 	FModuleManager::Get().LoadModule("ProjectLauncher");
 	FModuleManager::Get().LoadModule("SessionFrontend");
 	FModuleManager::Get().LoadModule("SettingsEditor");
-	FModuleManager::Get().LoadModule("CustomLaunchUI");
+	IPluginManager::Get().LoadModulesForEnabledPlugins(ELoadingPhase::Default);
 
 	InitializeSlateApplication(UnrealFrontendLayoutIni);
 
