Include a resource file to allow message dialogs to work in UnrealFrontend.

#rnx

diff --git a/Engine/Source/Programs/UnrealFrontend/Resources/Windows/UnrealFrontend.rc b/Engine/Source/Programs/UnrealFrontend/Resources/Windows/UnrealFrontend.rc
new file mode 100644
--- /dev/null
+++ b/Engine/Source/Programs/UnrealFrontend/Resources/Windows/UnrealFrontend.rc
@@ -0,0 +1,6 @@
+// Copyright Epic Games, Inc. All Rights Reserved.
+
+#include <windows.h>
+#include "Runtime/Launch/Resources/Version.h"
+#include "Runtime/Launch/Resources/Windows/PCLaunch.rc"
+
