fix UnrealFrontend icon

#rnx

diff --git a/Engine/Source/Programs/UnrealFrontend/Resources/Windows/UnrealFrontend.rc b/Engine/Source/Programs/UnrealFrontend/Resources/Windows/UnrealFrontend.rc
--- a/Engine/Source/Programs/UnrealFrontend/Resources/Windows/UnrealFrontend.rc
+++ b/Engine/Source/Programs/UnrealFrontend/Resources/Windows/UnrealFrontend.rc
@@ -4,3 +4,6 @@
 #include "Runtime/Launch/Resources/Version.h"
 #include "Runtime/Launch/Resources/Windows/PCLaunch.rc"
 
+#if !defined(AFX_RESOURCE_DLL) || defined(AFX_TARG_ENU)
+IDICON_UEGame ICON "UnrealFrontend.ico"
+#endif
