Add console command support to UnrealFrontend CLI

#Synced-CL 48294064

diff --git a/Engine/Source/Programs/UnrealFrontend/Private/Commands/ExecuteConsoleCommand.cpp b/Engine/Source/Programs/UnrealFrontend/Private/Commands/ExecuteConsoleCommand.cpp
new file mode 100644
--- /dev/null
+++ b/Engine/Source/Programs/UnrealFrontend/Private/Commands/ExecuteConsoleCommand.cpp
@@ -0,0 +1,108 @@
+﻿// Copyright Epic Games, Inc. All Rights Reserved.
+
+#include "ExecuteConsoleCommand.h"
+#include "ISessionServicesModule.h"
+#include "UnrealFrontendMain.h"
+#include "Async/TaskGraphInterfaces.h"
+#include "Containers/Ticker.h"
+#include "HAL/PlatformProcess.h"
+#include "HAL/PlatformTime.h"
+#include "Interfaces/IPluginManager.h"
+#include "Misc/CommandLine.h"
+#include "Misc/OutputDeviceRedirector.h"
+#include "Modules/ModuleManager.h"
+
+void FExecuteConsoleCommand::Run(const FString& Params)
+{
+	FString TargetInstanceIdStr = {};
+	FGuid TargetInstanceId;
+	if (!FParse::Value(FCommandLine::Get(), TEXT("-InstanceId="), TargetInstanceIdStr) || TargetInstanceIdStr.IsEmpty())
+	{
+		UE_LOG(LogUFECommands, Warning, TEXT("InstanceId was not set. Please use '-InstanceId=' in your command line. Tip: -InstanceId= can also be passed to any unreal application to override it's instance id."));
+		return;
+	}
+	if (!FGuid::Parse(TargetInstanceIdStr, TargetInstanceId))
+	{
+		UE_LOG(LogUFECommands, Warning, TEXT("Unable to parse InstanceId FGuid."));
+		return;
+	}
+
+	FString TargetCommand;
+	if (!FParse::Value(FCommandLine::Get(), TEXT("-Command="), TargetCommand) || TargetCommand.IsEmpty())
+	{
+		UE_LOG(LogUFECommands, Warning, TEXT("Command was not found. Please use '-Command=' in your command line."));
+		return;
+	}
+
+	double Timeout = 30.0;
+	FParse::Value(FCommandLine::Get(), TEXT("-Timeout="), Timeout);
+
+	// UdpMessaging will not start it's thread without correct loading phase
+	IPluginManager::Get().LoadModulesForEnabledPlugins(ELoadingPhase::PostDefault);
+
+	FModuleManager::Get().LoadModuleChecked("UdpMessaging");
+
+	TSharedPtr<ISessionManager> SessionManager = FModuleManager::Get().LoadModuleChecked<ISessionServicesModule>("SessionServices").GetSessionManager();
+
+	double Start = FPlatformTime::Seconds();
+
+	TSet<FGuid> DiscoveredInstanceIds;
+	bool bTargetCommandExecuted = false;
+
+	while (!IsEngineExitRequested() && (FPlatformTime::Seconds() - Start < Timeout) && !bTargetCommandExecuted)
+	{
+		TArray<TSharedPtr<ISessionInfo>> Sessions;
+		SessionManager->GetSessions(Sessions);
+		for (TSharedPtr<ISessionInfo> Session : Sessions)
+		{
+			TArray<TSharedPtr<ISessionInstanceInfo>> Instances;
+			Session->GetInstances(Instances);
+
+			for (TSharedPtr<ISessionInstanceInfo> Instance : Instances)
+			{
+				const FGuid& InstanceId = Instance->GetInstanceId();
+				const FString InstanceIdStr = InstanceId.ToString(EGuidFormats::DigitsWithHyphensLower);
+				if (!DiscoveredInstanceIds.Contains(InstanceId))
+				{
+					UE_LOG(LogUFECommands, Display, TEXT("Discovered InstanceId:%s InstanceName:%s Owner:%s SessionId:%s"),
+						*InstanceIdStr,
+						*Instance->GetInstanceName(),
+						*Session->GetSessionOwner(),
+						*Session->GetSessionId().ToString(EGuidFormats::DigitsWithHyphensLower));
+
+					DiscoveredInstanceIds.Add(InstanceId);
+				}
+
+				if (Instance->GetEngineVersion() == 0) // wait until instance has responded
+				{
+					continue;
+				}
+
+				if (InstanceId == TargetInstanceId)
+				{
+					if (Instance->IsAuthorized())
+					{
+						Instance->ExecuteCommand(TargetCommand);
+						UE_LOG(LogUFECommands, Warning, TEXT("Executed '%s' on InstanceId:%s"), *TargetCommand, *InstanceIdStr);
+					}
+					else
+					{
+						UE_LOG(LogUFECommands, Warning, TEXT("Not authorized to interact with InstanceId:%s"), *InstanceIdStr);
+					}
+
+					bTargetCommandExecuted = true;
+					break;
+				}
+			}
+		}
+
+		FTaskGraphInterface::Get().ProcessThreadUntilIdle(ENamedThreads::GameThread);
+		FTaskGraphInterface::Get().ProcessThreadUntilIdle(ENamedThreads::GameThread_Local);
+
+		FTSTicker::GetCoreTicker().Tick(0.1f);
+
+		FPlatformProcess::Sleep(0.1f);
+
+		GLog->FlushThreadedLogs();
+	}
+}
diff --git a/Engine/Source/Programs/UnrealFrontend/Private/Commands/ExecuteConsoleCommand.h b/Engine/Source/Programs/UnrealFrontend/Private/Commands/ExecuteConsoleCommand.h
new file mode 100644
--- /dev/null
+++ b/Engine/Source/Programs/UnrealFrontend/Private/Commands/ExecuteConsoleCommand.h
@@ -0,0 +1,20 @@
+﻿// Copyright Epic Games, Inc. All Rights Reserved.
+
+#pragma once
+
+#include "CoreMinimal.h"
+
+class ISessionInstanceInfo;
+class ISessionInfo;
+class ISessionManager;
+
+class FExecuteConsoleCommand
+{
+public:
+	/**
+	* Executes the command.
+	* Use '-Command=' command line argument to specify the console command.
+	* Use '-InstanceId=' command line argument to specify the target instance id.
+	*/
+	void Run(const FString& Params);
+};
diff --git a/Engine/Source/Programs/UnrealFrontend/Private/Commands/LaunchFromProfileCommand.cpp b/Engine/Source/Programs/UnrealFrontend/Private/Commands/LaunchFromProfileCommand.cpp
--- a/Engine/Source/Programs/UnrealFrontend/Private/Commands/LaunchFromProfileCommand.cpp
+++ b/Engine/Source/Programs/UnrealFrontend/Private/Commands/LaunchFromProfileCommand.cpp
@@ -9,12 +9,10 @@
 #include "ILauncher.h"
 #include "ILauncherServicesModule.h"
 #include "ITargetDeviceServicesModule.h"
-#include "UserInterfaceCommand.h"
+#include "UnrealFrontendMain.h"
 #include "Misc/CommandLine.h"
 #include "Misc/OutputDeviceRedirector.h"
 
-DEFINE_LOG_CATEGORY_STATIC(LogUFECommands, All, All);
-
 void FLaunchFromProfileCommand::Run(const FString& Params)
 {
 	// Get the name of the profile from the command line.
diff --git a/Engine/Source/Programs/UnrealFrontend/Private/UnrealFrontendMain.cpp b/Engine/Source/Programs/UnrealFrontend/Private/UnrealFrontendMain.cpp
--- a/Engine/Source/Programs/UnrealFrontend/Private/UnrealFrontendMain.cpp
+++ b/Engine/Source/Programs/UnrealFrontend/Private/UnrealFrontendMain.cpp
@@ -6,11 +6,12 @@
 #include "StatsConvertCommand.h"
 #include "UserInterfaceCommand.h"
 #include "LaunchFromProfileCommand.h"
+#include "ExecuteConsoleCommand.h"
 
+DEFINE_LOG_CATEGORY(LogUFECommands);
 
 IMPLEMENT_APPLICATION(UnrealFrontend, "UnrealFrontend");
 
-
 /**
  * Platform agnostic implementation of the main entry point.
  */
@@ -81,6 +82,11 @@ int32 UnrealFrontendMain( const TCHAR* CommandLine )
 			FLaunchFromProfileCommand* ProfileLaunch = new FLaunchFromProfileCommand;
 			ProfileLaunch->Run(Params);
 		}
+		else if (Command.Equals(TEXT("CONSOLE"), ESearchCase::IgnoreCase))
+		{
+			FExecuteConsoleCommand* ConsoleCommand = new FExecuteConsoleCommand;
+			ConsoleCommand->Run(Params);
+		}
 	}
 	else
 	{
diff --git a/Engine/Source/Programs/UnrealFrontend/Private/UnrealFrontendMain.h b/Engine/Source/Programs/UnrealFrontend/Private/UnrealFrontendMain.h
--- a/Engine/Source/Programs/UnrealFrontend/Private/UnrealFrontendMain.h
+++ b/Engine/Source/Programs/UnrealFrontend/Private/UnrealFrontendMain.h
@@ -4,6 +4,8 @@
 
 #include "CoreMinimal.h"
 
+DECLARE_LOG_CATEGORY_EXTERN(LogUFECommands, All, All);
+
 /**
  * The application's main function.
  *
diff --git a/Engine/Source/Runtime/Core/Private/Misc/App.cpp b/Engine/Source/Runtime/Core/Private/Misc/App.cpp
--- a/Engine/Source/Runtime/Core/Private/Misc/App.cpp
+++ b/Engine/Source/Runtime/Core/Private/Misc/App.cpp
@@ -28,6 +28,7 @@ DEFINE_LOG_CATEGORY_STATIC(LogApp, Log, All);
 bool FApp::bIsDebugGame = false;
 #endif
 
+FGuid FApp::InstanceId = FGuid::NewGuid();
 FGuid FApp::SessionId = ToGuid(FApp::GetSessionObjectId());
 FString FApp::SessionName = FString();
 FString FApp::SessionOwner = FString();
@@ -177,7 +178,6 @@ void FApp::InitializeSession()
 {
 	// parse session details on command line
 	FString InstanceIdString;
-	FGuid InstanceId = GetInstanceId();
 
 	if (FParse::Value(FCommandLine::Get(), TEXT("-InstanceId="), InstanceIdString))
 	{
@@ -188,7 +188,6 @@ void FApp::InitializeSession()
 	}
 
 	FString SessionIdString;
-
 	if (FParse::Value(FCommandLine::Get(), TEXT("-SessionId="), SessionIdString))
 	{
 		if (FGuid::Parse(SessionIdString, SessionId))
@@ -485,7 +484,6 @@ FString FApp::GetZenStoreProjectIdForProject(FStringView ProjectFilePath, FStrin
 
 FGuid FApp::GetInstanceId()
 {
-	static FGuid InstanceId = FGuid::NewGuid();
 	return InstanceId;
 }
 
diff --git a/Engine/Source/Runtime/Core/Public/Misc/App.h b/Engine/Source/Runtime/Core/Public/Misc/App.h
--- a/Engine/Source/Runtime/Core/Public/Misc/App.h
+++ b/Engine/Source/Runtime/Core/Public/Misc/App.h
@@ -871,6 +871,9 @@ private:
 	static CORE_API bool bIsDebugGame;
 #endif
 
+	/** Holds the instance identifier. */
+	static CORE_API FGuid InstanceId;
+
 	/** Holds the session identifier. */
 	static CORE_API FGuid SessionId;
 
diff --git a/Engine/Source/Runtime/SessionServices/Private/SessionManager.cpp b/Engine/Source/Runtime/SessionServices/Private/SessionManager.cpp
--- a/Engine/Source/Runtime/SessionServices/Private/SessionManager.cpp
+++ b/Engine/Source/Runtime/SessionServices/Private/SessionManager.cpp
@@ -37,13 +37,18 @@ FSessionManager::FSessionManager(const TSharedRef<IMessageBus, ESPMode::ThreadSa
 		FilteredOwners.Add(Filter);
 	}
 
+	if (!FParse::Value(FCommandLine::Get(), TEXT("SessionPingInterval="), PingInterval))
+	{
+		PingInterval = SESSION_MANAGER_PING_INTERVAL;
+	}
+
 	// connect to message bus
 	MessageEndpoint = FMessageEndpoint::Builder("FSessionManager", InMessageBus)
 		.Handling<FEngineServicePong>(this, &FSessionManager::HandleEnginePongMessage)
 		.Handling<FSessionServicePong>(this, &FSessionManager::HandleSessionPongMessage);
 
 	// initialize ticker
-	TickDelegateHandle = FTSTicker::GetCoreTicker().AddTicker(FTickerDelegate::CreateRaw(this, &FSessionManager::HandleTicker), SESSION_MANAGER_PING_INTERVAL);
+	TickDelegateHandle = FTSTicker::GetCoreTicker().AddTicker(FTickerDelegate::CreateRaw(this, &FSessionManager::HandleTicker), PingInterval);
 
 	SendPing();
 }
@@ -338,7 +343,7 @@ bool FSessionManager::HandleTicker(float DeltaTime)
 	// @todo gmp: don't expire sessions for now
 //	FindExpiredSessions(Now);
 
-	if (Now >= LastPingTime + FTimespan::FromSeconds(SESSION_MANAGER_PING_INTERVAL))
+	if (Now >= LastPingTime + FTimespan::FromSeconds(PingInterval))
 	{
 		SendPing();
 	}
diff --git a/Engine/Source/Runtime/SessionServices/Private/SessionManager.h b/Engine/Source/Runtime/SessionServices/Private/SessionManager.h
--- a/Engine/Source/Runtime/SessionServices/Private/SessionManager.h
+++ b/Engine/Source/Runtime/SessionServices/Private/SessionManager.h
@@ -162,4 +162,7 @@ private:
 
 	/** Holds a delegate to be invoked when the widget ticks. */
 	FTSTicker::FDelegateHandle TickDelegateHandle;
+
+	/** Holds current ping interval. */
+	float PingInterval;
 };
